Interchange test framework. This implements a framework for testing asset imports, but it is designed with the idea of extending it to testing scene import, and asset exports.

diff --git a/Engine/Plugins/Tests/InterchangeTests/InterchangeTests.uplugin b/Engine/Plugins/Tests/InterchangeTests/InterchangeTests.uplugin
new file mode 100644
--- /dev/null
+++ b/Engine/Plugins/Tests/InterchangeTests/InterchangeTests.uplugin
@@ -0,0 +1,40 @@
+{
+	"FileVersion": 3,
+	"Version": 1,
+	"VersionName": "1.0",
+	"FriendlyName": "Interchange Tests",
+	"Description": "Plugin for Interchange automation tests.",
+	"Category": "Testing",
+	"CreatedBy": "Epic Games, Inc.",
+	"CreatedByURL": "http://epicgames.com",
+	"DocsURL": "",
+	"MarketplaceURL": "",
+	"SupportURL": "",
+	"EnabledByDefault": true,
+	"CanContainContent": false,
+	"IsBetaVersion": true,
+	"Installed": false,
+	"IsExperimentalVersion": true,
+	"Modules": [
+		{
+			"Name": "InterchangeTests",
+			"Type": "Editor",
+			"LoadingPhase": "Default"
+		},
+		{
+			"Name": "InterchangeTestEditor",
+			"Type": "Editor",
+			"LoadingPhase": "Default"
+		}
+	],
+	"Plugins": [
+		{
+			"Name": "Interchange",
+			"Enabled": true,
+			"TargetAllowList": [
+				"Editor",
+				"Program"
+			]
+		}
+	]
+}
diff --git a/Engine/Plugins/Tests/InterchangeTests/Source/InterchangeTestEditor/InterchangeTestEditor.Build.cs b/Engine/Plugins/Tests/InterchangeTests/Source/InterchangeTestEditor/InterchangeTestEditor.Build.cs
new file mode 100644
--- /dev/null
+++ b/Engine/Plugins/Tests/InterchangeTests/Source/InterchangeTestEditor/InterchangeTestEditor.Build.cs
@@ -0,0 +1,37 @@
+// Copyright Epic Games, Inc. All Rights Reserved.
+
+using System.IO;
+
+namespace UnrealBuildTool.Rules
+{
+	public class InterchangeTestEditor : ModuleRules
+	{
+		public InterchangeTestEditor(ReadOnlyTargetRules Target) : base(Target)
+		{
+            PublicDependencyModuleNames.AddRange(
+				new string[]
+				{
+                    "Core",
+					"CoreUObject",
+					"Engine",
+					"EditorStyle",
+					"InputCore",
+					"InterchangeTests",
+					"LevelEditor",
+					"PropertyEditor",
+					"UnrealEd"
+				}
+			);
+
+			PrivateDependencyModuleNames.AddRange(
+				new string[]
+				{
+					"AssetTools",
+					"WorkspaceMenuStructure",
+					"SlateCore",
+					"Slate"
+				}
+			);
+		}
+    }
+}
diff --git a/Engine/Plugins/Tests/InterchangeTests/Source/InterchangeTestEditor/Private/AssetTypeActions_InterchangeImportTestPlan.cpp b/Engine/Plugins/Tests/InterchangeTests/Source/InterchangeTestEditor/Private/AssetTypeActions_InterchangeImportTestPlan.cpp
new file mode 100644
--- /dev/null
+++ b/Engine/Plugins/Tests/InterchangeTests/Source/InterchangeTestEditor/Private/AssetTypeActions_InterchangeImportTestPlan.cpp
@@ -0,0 +1,19 @@
+// Copyright Epic Games, Inc. All Rights Reserved.
+
+#include "AssetTypeActions_InterchangeImportTestPlan.h"
+
+#define LOCTEXT_NAMESPACE "AssetTypeActions"
+
+
+void FAssetTypeActions_InterchangeImportTestPlan::OpenAssetEditor(const TArray<UObject*>& InObjects, TSharedPtr<IToolkitHost> EditWithinLevelEditor)
+{
+	for (UObject* Object : InObjects)
+	{
+		if (UInterchangeImportTestPlan* TestPlan = Cast<UInterchangeImportTestPlan>(Object))
+		{
+			FSimpleAssetEditor::CreateEditor(EToolkitMode::Standalone, EditWithinLevelEditor, Object);
+		}
+	}
+}
+
+#undef LOCTEXT_NAMESPACE
diff --git a/Engine/Plugins/Tests/InterchangeTests/Source/InterchangeTestEditor/Private/AssetTypeActions_InterchangeImportTestPlan.h b/Engine/Plugins/Tests/InterchangeTests/Source/InterchangeTestEditor/Private/AssetTypeActions_InterchangeImportTestPlan.h
new file mode 100644
--- /dev/null
+++ b/Engine/Plugins/Tests/InterchangeTests/Source/InterchangeTestEditor/Private/AssetTypeActions_InterchangeImportTestPlan.h
@@ -0,0 +1,22 @@
+// Copyright Epic Games, Inc. All Rights Reserved.
+
+#pragma once
+
+#include "CoreMinimal.h"
+#include "AssetTypeActions_Base.h"
+#include "InterchangeImportTestPlan.h"
+
+#define LOCTEXT_NAMESPACE "AssetTypeActions"
+
+class FAssetTypeActions_InterchangeImportTestPlan : public FAssetTypeActions_Base
+{
+public:
+	// IAssetTypeActions Implementation
+	virtual FText GetName() const override { return LOCTEXT("AssetTypeActions_InterchangeImportTestPlan", "Interchange Import Test Plan"); }
+	virtual FColor GetTypeColor() const override { return FColor(125, 0, 200); }
+	virtual UClass* GetSupportedClass() const override { return UInterchangeImportTestPlan::StaticClass(); }
+	virtual uint32 GetCategories() override { return EAssetTypeCategories::Misc; }
+	virtual void OpenAssetEditor(const TArray<UObject*>& InObjects, TSharedPtr<class IToolkitHost> EditWithinLevelEditor = TSharedPtr<IToolkitHost>()) override;
+};
+
+#undef LOCTEXT_NAMESPACE
diff --git a/Engine/Plugins/Tests/InterchangeTests/Source/InterchangeTestEditor/Private/InterchangeImportTestPlanFactory.cpp b/Engine/Plugins/Tests/InterchangeTests/Source/InterchangeTestEditor/Private/InterchangeImportTestPlanFactory.cpp
new file mode 100644
--- /dev/null
+++ b/Engine/Plugins/Tests/InterchangeTests/Source/InterchangeTestEditor/Private/InterchangeImportTestPlanFactory.cpp
@@ -0,0 +1,38 @@
+// Copyright Epic Games, Inc. All Rights Reserved.
+
+#include "InterchangeImportTestPlanFactory.h"
+#include "InterchangeImportTestPlan.h"
+#include "AssetTypeCategories.h"
+
+#define LOCTEXT_NAMESPACE "InterchangeImportTestPlanFactory"
+
+
+UInterchangeImportTestPlanFactory::UInterchangeImportTestPlanFactory()
+{
+	bCreateNew = true;
+	bEditAfterNew = false;
+	bEditorImport = false;
+	SupportedClass = UInterchangeImportTestPlan::StaticClass();
+}
+
+
+UObject* UInterchangeImportTestPlanFactory::FactoryCreateNew(UClass* Class, UObject* InParent, FName InName, EObjectFlags InFlags, UObject* Context, FFeedbackContext* Warn)
+{
+	UInterchangeImportTestPlan* TestPlan = NewObject<UInterchangeImportTestPlan>(InParent, SupportedClass, InName, InFlags | RF_Transactional);
+	return TestPlan;
+}
+
+
+uint32 UInterchangeImportTestPlanFactory::GetMenuCategories() const
+{
+	return EAssetTypeCategories::Misc;
+}
+
+
+FText UInterchangeImportTestPlanFactory::GetDisplayName() const
+{
+	return LOCTEXT("MenuEntry", "Interchange Import Test Plan");
+}
+
+
+#undef LOCTEXT_NAMESPACE
diff --git a/Engine/Plugins/Tests/InterchangeTests/Source/InterchangeTestEditor/Private/InterchangeImportTestPlanFactory.h b/Engine/Plugins/Tests/InterchangeTests/Source/InterchangeTestEditor/Private/InterchangeImportTestPlanFactory.h
new file mode 100644
--- /dev/null
+++ b/Engine/Plugins/Tests/InterchangeTests/Source/InterchangeTestEditor/Private/InterchangeImportTestPlanFactory.h
@@ -0,0 +1,24 @@
+// Copyright Epic Games, Inc. All Rights Reserved.
+
+#pragma once
+
+#include "CoreMinimal.h"
+#include "UObject/ObjectMacros.h"
+#include "Factories/Factory.h"
+#include "InterchangeImportTestPlanFactory.generated.h"
+
+
+UCLASS(hidecategories=Object, collapsecategories)
+class UInterchangeImportTestPlanFactory : public UFactory
+{
+	GENERATED_BODY()
+
+public:
+	UInterchangeImportTestPlanFactory();
+
+	// UFactory Interface
+	virtual UObject* FactoryCreateNew(UClass* Class, UObject* InParent, FName Name, EObjectFlags Flags, UObject* Context, FFeedbackContext* Warn) override;
+	virtual uint32 GetMenuCategories() const override;
+	virtual FText GetDisplayName() const override;
+	// End of UFactory Interface
+};
diff --git a/Engine/Plugins/Tests/InterchangeTests/Source/InterchangeTestEditor/Private/InterchangeTestEditorModule.cpp b/Engine/Plugins/Tests/InterchangeTests/Source/InterchangeTestEditor/Private/InterchangeTestEditorModule.cpp
new file mode 100644
--- /dev/null
+++ b/Engine/Plugins/Tests/InterchangeTests/Source/InterchangeTestEditor/Private/InterchangeTestEditorModule.cpp
@@ -0,0 +1,64 @@
+// Copyright Epic Games, Inc. All Rights Reserved.
+
+#include "InterchangeTestEditorModule.h"
+#include "InterchangeImportTestPlan.h"
+#include "InterchangeImportTestSettings.h"
+#include "InterchangeTestsModule.h"
+#include "Modules/ModuleInterface.h"
+#include "Modules/ModuleManager.h"
+#include "InterchangeTestFunctionLayout.h"
+#include "InterchangeTestFunction.h"
+#include "AssetTypeActions_InterchangeImportTestPlan.h"
+#include "AssetToolsModule.h"
+
+
+#define LOCTEXT_NAMESPACE "InterchangeTestEditorModule"
+
+
+FInterchangeTestEditorModule& FInterchangeTestEditorModule::Get()
+{
+	return FModuleManager::LoadModuleChecked<FInterchangeTestEditorModule>(INTERCHANGETESTEDITOR_MODULE_NAME);
+}
+
+
+bool FInterchangeTestEditorModule::IsAvailable()
+{
+	return FModuleManager::Get().IsModuleLoaded(INTERCHANGETESTEDITOR_MODULE_NAME);
+}
+
+
+void FInterchangeTestEditorModule::StartupModule()
+{
+	// Register the InterchangeImportTestPlan asset
+	FAssetToolsModule& AssetToolsModule = FAssetToolsModule::GetModule();
+	IAssetTools& AssetTools = AssetToolsModule.Get();
+	AssetTypeActions = MakeShared<FAssetTypeActions_InterchangeImportTestPlan>();
+	AssetTools.RegisterAssetTypeActions(AssetTypeActions.ToSharedRef());
+
+	// Register the FInterchangeTestFunction struct customization
+	FPropertyEditorModule& PropertyModule = FModuleManager::LoadModuleChecked<FPropertyEditorModule>("PropertyEditor");
+	PropertyModule.RegisterCustomPropertyTypeLayout(FInterchangeTestFunction::StaticStruct()->GetFName(),
+	FOnGetPropertyTypeCustomizationInstance::CreateStatic(&FInterchangeTestFunctionLayout::MakeInstance));
+	PropertyModule.NotifyCustomizationModuleChanged();
+}
+
+
+void FInterchangeTestEditorModule::ShutdownModule()
+{
+	if (UObjectInitialized())
+	{
+		FAssetToolsModule& AssetToolsModule = FAssetToolsModule::GetModule();
+		IAssetTools& AssetTools = AssetToolsModule.Get();
+		AssetTools.UnregisterAssetTypeActions(AssetTypeActions.ToSharedRef());
+
+		FPropertyEditorModule& PropertyModule = FModuleManager::LoadModuleChecked<FPropertyEditorModule>("PropertyEditor");
+		PropertyModule.UnregisterCustomPropertyTypeLayout(FInterchangeTestFunction::StaticStruct()->GetFName());
+	}
+}
+
+
+IMPLEMENT_MODULE(FInterchangeTestEditorModule, InterchangeTestEditor);
+
+
+#undef LOCTEXT_NAMESPACE
+
diff --git a/Engine/Plugins/Tests/InterchangeTests/Source/InterchangeTestEditor/Private/InterchangeTestFunctionLayout.cpp b/Engine/Plugins/Tests/InterchangeTests/Source/InterchangeTestEditor/Private/InterchangeTestFunctionLayout.cpp
new file mode 100644
--- /dev/null
+++ b/Engine/Plugins/Tests/InterchangeTests/Source/InterchangeTestEditor/Private/InterchangeTestFunctionLayout.cpp
@@ -0,0 +1,273 @@
+// Copyright Epic Games, Inc. All Rights Reserved.
+
+#include "InterchangeTestFunctionLayout.h"
+#include "IDetailChildrenBuilder.h"
+#include "IPropertyUtilities.h"
+#include "DetailWidgetRow.h"
+#include "DetailLayoutBuilder.h"
+#include "Widgets/Input/SNumericEntryBox.h"
+#include "Widgets/Input/SComboBox.h"
+#include "InterchangeTestFunction.h"
+
+#define LOCTEXT_NAMESPACE "InterchangeTestFunctionLayout"
+
+TSharedRef<IPropertyTypeCustomization> FInterchangeTestFunctionLayout::MakeInstance()
+{
+	return MakeShared<FInterchangeTestFunctionLayout>();
+}
+
+
+FInterchangeTestFunctionLayout::FInterchangeTestFunctionLayout()
+{
+}
+
+
+FInterchangeTestFunctionLayout::~FInterchangeTestFunctionLayout()
+{
+	FEditorDelegates::PostUndoRedo.RemoveAll(this);
+}
+
+
+void FInterchangeTestFunctionLayout::CustomizeHeader(TSharedRef<IPropertyHandle> StructPropertyHandle, FDetailWidgetRow& HeaderRow, IPropertyTypeCustomizationUtils& StructCustomizationUtils)
+{
+	// Perform initialization here rather than in the constructor
+	// (adding SP delegates doesn't work if the object hasn't yet been assigned)
+
+	AssetClasses = FInterchangeTestFunction::GetAvailableAssetClasses();
+	FEditorDelegates::PostUndoRedo.AddSP(this, &FInterchangeTestFunctionLayout::RefreshLayout);
+
+	// Ensure that the details layout is refreshed when resetting the entire struct to default, so that custom data can be rebuilt
+
+	StructProperty = StructPropertyHandle;
+	StructProperty->SetOnPropertyResetToDefault(FSimpleDelegate::CreateSP(this, &FInterchangeTestFunctionLayout::RefreshLayout));
+
+	// Cache all the child properties we're interested in manipulating
+
+	AssetClassProperty = StructPropertyHandle->GetChildHandle(GET_MEMBER_NAME_CHECKED(FInterchangeTestFunction, AssetClass));
+	CheckFunctionProperty = StructPropertyHandle->GetChildHandle(GET_MEMBER_NAME_CHECKED(FInterchangeTestFunction, CheckFunction));
+	ParametersProperty = StructPropertyHandle->GetChildHandle(GET_MEMBER_NAME_CHECKED(FInterchangeTestFunction, Parameters));
+	PropertyUtilities = StructCustomizationUtils.GetPropertyUtilities();
+
+	// The FInterchangeTestFunction defines a specific test which will be carried out on the imported asset.
+	// The three important properties are:
+	// 1) AssetClass - the class of the asset object we are expecting, and which we will perform our test on.
+	// 2) TestFunction - a UFunction* specifying the function we will call to perform the test
+	// 3) Parameters - a list of specific parameter values to call the function with
+	//
+	// The Parameters property is a simple (name, valueString) map.
+	// However this is not a useful form for interacting with in a properties view, so we import the parameter values to binary here.
+	// Note that this binary blob is owned by, and only of interest to, this layout class.
+
+	FInterchangeTestFunction* InterchangeTestFunction = GetStruct();
+	GetFunctionsForClass(InterchangeTestFunction->AssetClass);
+	ParamData = InterchangeTestFunction->ImportParameters();
+
+	// The header of the struct customization contains the class and function combo boxes.
+	// These will determine which children should be populated underneath.
+
+	HeaderRow
+	.NameContent()
+	[
+		StructPropertyHandle->CreatePropertyNameWidget()
+	]
+	.ValueContent()
+	[
+		SNew(SHorizontalBox)
+		+SHorizontalBox::Slot()
+		.AutoWidth()
+		[
+			SNew(SComboBox<UClass*>)
+			.OptionsSource(&AssetClasses)
+			.OnGenerateWidget(this, &FInterchangeTestFunctionLayout::OnGenerateClassComboWidget)
+			.OnSelectionChanged(this, &FInterchangeTestFunctionLayout::OnClassComboSelectionChanged)
+			.InitiallySelectedItem(InterchangeTestFunction->AssetClass)
+			[
+				OnGenerateClassComboWidget(InterchangeTestFunction->AssetClass)
+			]
+		]
+		+SHorizontalBox::Slot()
+		.AutoWidth()
+		[
+			SNew(SComboBox<UFunction*>)
+			.OptionsSource(&Functions)
+			.OnGenerateWidget(this, &FInterchangeTestFunctionLayout::OnGenerateFunctionComboWidget)
+			.OnSelectionChanged(this, &FInterchangeTestFunctionLayout::OnFunctionComboSelectionChanged)
+			.InitiallySelectedItem(InterchangeTestFunction->CheckFunction)
+			[
+				OnGenerateFunctionComboWidget(InterchangeTestFunction->CheckFunction)
+			]
+		]
+	];
+}
+
+
+void FInterchangeTestFunctionLayout::GetFunctionsForClass(UClass* AssetClass)
+{
+	// This populates an array with all the UFunction* available as tests for the given asset class
+	Functions = FInterchangeTestFunction::GetAvailableFunctions(AssetClass);
+}
+
+
+void FInterchangeTestFunctionLayout::CustomizeChildren(TSharedRef<IPropertyHandle> StructPropertyHandle, IDetailChildrenBuilder& StructBuilder, IPropertyTypeCustomizationUtils& StructCustomizationUtils)
+{
+	// The children of the struct customization represent all the editable input parameters to bind to the test function
+	// This needs to be rebuilt every time a new function is selected.
+
+	FInterchangeTestFunction* InterchangeTestFunction = GetStruct();
+
+	// Iterate through the parameters of interest
+	for (FProperty* ParamProperty : InterchangeTestFunction->GetParameters())
+	{
+		FName ParamName = ParamProperty->GetFName();
+		IDetailPropertyRow* Row = StructBuilder.AddExternalStructureProperty(ParamData.ToSharedRef(), ParamName, FAddPropertyParams().ForceShowProperty());
+		TSharedPtr<IPropertyHandle> PropertyHandle = Row->GetPropertyHandle();
+
+		// Note: we need to hook both of these events.
+		// SetOnPropertyValueChanged is called when the property itself is modified.
+		// SetOnChildPropertyValueChanged is called when a child property is modified, e.g. the component of a vector.
+
+		PropertyHandle->SetOnPropertyValueChanged(FSimpleDelegate::CreateSP(this, &FInterchangeTestFunctionLayout::OnParameterChanged, PropertyHandle));
+		PropertyHandle->SetOnChildPropertyValueChanged(FSimpleDelegate::CreateSP(this, &FInterchangeTestFunctionLayout::OnParameterChanged, PropertyHandle));
+
+		// We have a notion of parameter defaults from the C++ function prototype.
+		// Hook a custom handler for Reset To Default for these parameter properties so that it does something useful.
+
+		Row->OverrideResetToDefault(FResetToDefaultOverride::Create(FSimpleDelegate::CreateSP(this, &FInterchangeTestFunctionLayout::OnParameterResetToDefault, PropertyHandle)));
+	}
+}
+
+
+FInterchangeTestFunction* FInterchangeTestFunctionLayout::GetStruct() const
+{
+	// Get address of the FInterchangeTestFunction struct being viewed.
+	// We only ever expect the property handle to be linked to a single instance.
+
+	void* StructPtr = nullptr;
+	StructProperty->GetValueData(StructPtr);
+	return static_cast<FInterchangeTestFunction*>(StructPtr);
+}
+
+
+void FInterchangeTestFunctionLayout::OnParameterChanged(TSharedPtr<IPropertyHandle> PropertyHandle)
+{
+	// When we change a function parameter property in the property editor, this will change the binary copy we maintain here.
+	// Here we need to propagate the change to the struct's (name, value) map.
+
+	FString Value;
+	PropertyHandle->GetValueAsFormattedString(Value);
+	FName PropertyName = PropertyHandle->GetProperty()->GetFName();
+
+	void* MapPtr = nullptr;
+	ParametersProperty->GetValueData(MapPtr);
+
+	// Note that we are not accessing the map property directly through the handle, so we need to call these various notifies ourselves.
+	// This allows transactions to work correctly.
+
+	ParametersProperty->NotifyPreChange();
+	static_cast<decltype(FInterchangeTestFunction::Parameters)*>(MapPtr)->Add(PropertyName, Value);
+	ParametersProperty->NotifyPostChange(EPropertyChangeType::ValueSet);
+	ParametersProperty->NotifyFinishedChangingProperties();
+}
+
+
+void FInterchangeTestFunctionLayout::OnParameterResetToDefault(TSharedPtr<IPropertyHandle> PropertyHandle)
+{
+	// The property editor doesn't start a transaction if ResetToDefault is custom. We have to do it ourselves!
+	FScopedTransaction Transaction(FText::Format(LOCTEXT("ResetParameterToDefault", "Reset {0} to Default"), PropertyHandle->GetPropertyDisplayName()));
+
+	FName PropertyName = PropertyHandle->GetProperty()->GetFName();
+
+	// The way we implement ResetToDefault is by just removing the parameter from the map completely.
+	// (note the requirement to call the notifies manually)
+	// The default value will be set in the binary copy we maintain when the parameter is imported.
+
+	void* MapPtr = nullptr;
+	ParametersProperty->GetValueData(MapPtr);
+
+	ParametersProperty->NotifyPreChange();
+	static_cast<decltype(FInterchangeTestFunction::Parameters)*>(MapPtr)->Remove(PropertyName);
+	ParametersProperty->NotifyPostChange(EPropertyChangeType::ValueSet);
+	ParametersProperty->NotifyFinishedChangingProperties();
+
+	FInterchangeTestFunction* InterchangeTestFunction = GetStruct();
+	InterchangeTestFunction->ImportParameter(*ParamData, PropertyHandle->GetProperty());
+}
+
+
+void FInterchangeTestFunctionLayout::RefreshLayout()
+{
+	if (PropertyUtilities.IsValid())
+	{
+		PropertyUtilities->ForceRefresh();
+	}
+}
+
+
+TSharedRef<SWidget> FInterchangeTestFunctionLayout::OnGenerateClassComboWidget(UClass* InClass)
+{
+	return SNew(SBox)
+	[
+		SNew(STextBlock)
+		.Text(InClass ? InClass->GetDisplayNameText() : LOCTEXT("SelectAssetClass", "<Expected asset class>"))
+		.ToolTipText(InClass ? InClass->GetToolTipText() : FText())
+		.Font(IDetailLayoutBuilder::GetDetailFont())
+	];
+}
+
+
+void FInterchangeTestFunctionLayout::OnClassComboSelectionChanged(UClass* InSelectedItem, ESelectInfo::Type SelectInfo)
+{
+	FInterchangeTestFunction* InterchangeTestFunction = GetStruct();
+
+	if (InterchangeTestFunction->AssetClass != InSelectedItem)
+	{
+		// When the asset class combo box of the header changes, this changes the entire state of the struct.
+		// The TestFunction needs to be reset (as the old one is most likely unsuitable), and so does the stored parameter map.
+		// So wrap this entire operation in a single high-level transaction, and force the layout to be refreshed afterwards.
+
+		FScopedTransaction Transaction(LOCTEXT("EditTestType", "Change test to perform"));
+		AssetClassProperty->SetValue(InSelectedItem);
+		UFunction* NullFunction = nullptr;
+		CheckFunctionProperty->SetValue(NullFunction);
+		ParametersProperty->AsMap()->Empty();
+
+		RefreshLayout();
+	}
+}
+
+
+TSharedRef<SWidget> FInterchangeTestFunctionLayout::OnGenerateFunctionComboWidget(UFunction* InFunction)
+{
+	FInterchangeTestFunction* InterchangeTestFunction = GetStruct();
+
+	// If a test function doesn't have the correct signature, it will appear in the list, but greyed out and unselectable.
+	// This will act as a signal to the developer to fix the incorrect signature.
+	return SNew(SBox)
+	[
+		SNew(STextBlock)
+		.Text(InFunction ? InFunction->GetDisplayNameText() : LOCTEXT("SelectTest", "<Please select a test>"))
+		.ToolTipText(InFunction ? InFunction->GetToolTipText() : FText())
+		.Font(IDetailLayoutBuilder::GetDetailFont())
+		.IsEnabled(InterchangeTestFunction->IsValid(InFunction))
+	];
+}
+
+
+void FInterchangeTestFunctionLayout::OnFunctionComboSelectionChanged(UFunction* InSelectedItem, ESelectInfo::Type SelectInfo)
+{
+	FInterchangeTestFunction* InterchangeTestFunction = GetStruct();
+
+	if (InterchangeTestFunction->IsValid(InSelectedItem) && InterchangeTestFunction->CheckFunction != InSelectedItem)
+	{
+		// When the function combo box of the header changes, we need to repopulate the parameter map.
+		// So wrap this entire operation in a single high-level transaction, and force the layout to be refreshed afterwards.
+
+		FScopedTransaction Transaction(LOCTEXT("EditTestType", "Change test to perform"));
+		CheckFunctionProperty->SetValue(InSelectedItem);
+		ParametersProperty->AsMap()->Empty();
+
+		RefreshLayout();
+	}
+}
+
+#undef LOCTEXT_NAMESPACE
diff --git a/Engine/Plugins/Tests/InterchangeTests/Source/InterchangeTestEditor/Private/InterchangeTestFunctionLayout.h b/Engine/Plugins/Tests/InterchangeTests/Source/InterchangeTestEditor/Private/InterchangeTestFunctionLayout.h
new file mode 100644
--- /dev/null
+++ b/Engine/Plugins/Tests/InterchangeTests/Source/InterchangeTestEditor/Private/InterchangeTestFunctionLayout.h
@@ -0,0 +1,56 @@
+// Copyright Epic Games, Inc. All Rights Reserved.
+
+#pragma once
+
+#include "CoreMinimal.h"
+#include "IPropertyTypeCustomization.h"
+#include "PropertyHandle.h"
+#include "EditorUndoClient.h"
+
+class FDetailWidgetRow;
+class IDetailChildrenBuilder;
+class IPropertyUtilities;
+struct FInterchangeTestFunction;
+
+
+class FInterchangeTestFunctionLayout : public IPropertyTypeCustomization
+{
+public:
+	/** Makes a new instance of this layout class for a specific detail view requesting it */
+	static TSharedRef<IPropertyTypeCustomization> MakeInstance();
+
+	FInterchangeTestFunctionLayout();
+	virtual ~FInterchangeTestFunctionLayout();
+
+	// IPropertyTypeCustomization interface
+	virtual void CustomizeHeader(TSharedRef<IPropertyHandle> StructPropertyHandle, FDetailWidgetRow& HeaderRow, IPropertyTypeCustomizationUtils& StructCustomizationUtils) override;
+	virtual void CustomizeChildren(TSharedRef<IPropertyHandle> StructPropertyHandle, IDetailChildrenBuilder& StructBuilder, IPropertyTypeCustomizationUtils& StructCustomizationUtils) override;
+
+private:
+	FInterchangeTestFunction* GetStruct() const;
+
+	void GetFunctionsForClass(UClass* Class);
+	TSharedRef<SWidget> OnGenerateClassComboWidget(UClass* InClass);
+	void OnClassComboSelectionChanged(UClass* InSelectedItem, ESelectInfo::Type SelectInfo);
+	TSharedRef<SWidget> OnGenerateFunctionComboWidget(UFunction* InFunction);
+	void OnFunctionComboSelectionChanged(UFunction* InSelectedItem, ESelectInfo::Type SelectInfo);
+	void OnParameterChanged(TSharedPtr<IPropertyHandle> PropertyHandle);
+	void OnParameterResetToDefault(TSharedPtr<IPropertyHandle> PropertyHandle);
+	void RefreshLayout();
+
+	/** List of all the asset classes supported by the testing framework */
+	TArray<UClass*> AssetClasses;
+
+	/** List of all the test functions supported by the current asset class */
+	TArray<UFunction*> Functions;
+
+	/** Bound parameter data stored as a binary blob */
+	TSharedPtr<FStructOnScope> ParamData;
+
+	TSharedPtr<IPropertyHandle> StructProperty;
+	TSharedPtr<IPropertyHandle> AssetClassProperty;
+	TSharedPtr<IPropertyHandle> CheckFunctionProperty;
+	TSharedPtr<IPropertyHandle> ParametersProperty;
+
+	TSharedPtr<IPropertyUtilities> PropertyUtilities;
+};
diff --git a/Engine/Plugins/Tests/InterchangeTests/Source/InterchangeTestEditor/Public/InterchangeTestEditorModule.h b/Engine/Plugins/Tests/InterchangeTests/Source/InterchangeTestEditor/Public/InterchangeTestEditorModule.h
new file mode 100644
--- /dev/null
+++ b/Engine/Plugins/Tests/InterchangeTests/Source/InterchangeTestEditor/Public/InterchangeTestEditorModule.h
@@ -0,0 +1,27 @@
+// Copyright Epic Games, Inc. All Rights Reserved.
+
+#pragma once
+
+#include "CoreMinimal.h"
+#include "Modules/ModuleInterface.h"
+
+class FAssetTypeActions_InterchangeImportTestPlan;
+
+
+#define INTERCHANGETESTEDITOR_MODULE_NAME TEXT("InterchangeTestEditor")
+
+/**
+ * Module for implementing the editor for Interchange automation tests
+ */
+class FInterchangeTestEditorModule : public IModuleInterface
+{
+public:
+	static FInterchangeTestEditorModule& Get();
+	static bool IsAvailable();
+
+	virtual void StartupModule() override;
+	virtual void ShutdownModule() override;
+
+private:
+	TSharedPtr<FAssetTypeActions_InterchangeImportTestPlan> AssetTypeActions;
+};
diff --git a/Engine/Plugins/Tests/InterchangeTests/Source/InterchangeTests/InterchangeTests.Build.cs b/Engine/Plugins/Tests/InterchangeTests/Source/InterchangeTests/InterchangeTests.Build.cs
new file mode 100644
--- /dev/null
+++ b/Engine/Plugins/Tests/InterchangeTests/Source/InterchangeTests/InterchangeTests.Build.cs
@@ -0,0 +1,36 @@
+// Copyright Epic Games, Inc. All Rights Reserved.
+
+using System.IO;
+
+namespace UnrealBuildTool.Rules
+{
+	public class InterchangeTests : ModuleRules
+	{
+		public InterchangeTests(ReadOnlyTargetRules Target) : base(Target)
+		{
+            PublicDependencyModuleNames.AddRange(
+				new string[]
+				{
+                    "Core",
+					"CoreUObject",
+					"Engine",
+					"Json",
+					"JsonUtilities"
+				}
+			);
+
+			PrivateDependencyModuleNames.AddRange(
+				new string[]
+				{
+					"AssetRegistry",
+					"InterchangePipelines",
+					"InterchangeCore",
+					"InterchangeEngine",
+					"MeshDescription",
+					"StaticMeshDescription",
+					"UnrealEd"
+				}
+			);
+		}
+    }
+}
diff --git a/Engine/Plugins/Tests/InterchangeTests/Source/InterchangeTests/Private/ImportTestFunctions/AnimationImportTestFunctions.cpp b/Engine/Plugins/Tests/InterchangeTests/Source/InterchangeTests/Private/ImportTestFunctions/AnimationImportTestFunctions.cpp
new file mode 100644
--- /dev/null
+++ b/Engine/Plugins/Tests/InterchangeTests/Source/InterchangeTests/Private/ImportTestFunctions/AnimationImportTestFunctions.cpp
@@ -0,0 +1,11 @@
+// Copyright Epic Games, Inc. All Rights Reserved.
+
+#include "ImportTestFunctions/AnimationImportTestFunctions.h"
+#include "InterchangeTestFunction.h"
+#include "Engine/SkeletalMesh.h"
+
+
+UClass* UAnimationImportTestFunctions::GetAssociatedAssetType() const
+{
+	return USkeletalMesh::StaticClass();
+}
diff --git a/Engine/Plugins/Tests/InterchangeTests/Source/InterchangeTests/Private/ImportTestFunctions/ImportTestFunctionsBase.cpp b/Engine/Plugins/Tests/InterchangeTests/Source/InterchangeTests/Private/ImportTestFunctions/ImportTestFunctionsBase.cpp
new file mode 100644
--- /dev/null
+++ b/Engine/Plugins/Tests/InterchangeTests/Source/InterchangeTests/Private/ImportTestFunctions/ImportTestFunctionsBase.cpp
@@ -0,0 +1,9 @@
+// Copyright Epic Games, Inc. All Rights Reserved.
+
+#include "ImportTestFunctions/ImportTestFunctionsBase.h"
+
+
+UClass* UImportTestFunctionsBase::GetImportTestFunctionsClassForAssetType(UClass* AssetClass)
+{
+	return nullptr;
+}
diff --git a/Engine/Plugins/Tests/InterchangeTests/Source/InterchangeTests/Private/ImportTestFunctions/InterchangeResultImportTestFunctions.cpp b/Engine/Plugins/Tests/InterchangeTests/Source/InterchangeTests/Private/ImportTestFunctions/InterchangeResultImportTestFunctions.cpp
new file mode 100644
--- /dev/null
+++ b/Engine/Plugins/Tests/InterchangeTests/Source/InterchangeTests/Private/ImportTestFunctions/InterchangeResultImportTestFunctions.cpp
@@ -0,0 +1,11 @@
+// Copyright Epic Games, Inc. All Rights Reserved.
+
+#include "ImportTestFunctions/InterchangeResultImportTestFunctions.h"
+#include "InterchangeTestFunction.h"
+#include "InterchangeResultsContainer.h"
+
+
+UClass* UInterchangeResultImportTestFunctions::GetAssociatedAssetType() const
+{
+	return UInterchangeResultsContainer::StaticClass();
+}
diff --git a/Engine/Plugins/Tests/InterchangeTests/Source/InterchangeTests/Private/ImportTestFunctions/MaterialImportTestFunctions.cpp b/Engine/Plugins/Tests/InterchangeTests/Source/InterchangeTests/Private/ImportTestFunctions/MaterialImportTestFunctions.cpp
new file mode 100644
--- /dev/null
+++ b/Engine/Plugins/Tests/InterchangeTests/Source/InterchangeTests/Private/ImportTestFunctions/MaterialImportTestFunctions.cpp
@@ -0,0 +1,24 @@
+// Copyright Epic Games, Inc. All Rights Reserved.
+
+#include "ImportTestFunctions/MaterialImportTestFunctions.h"
+#include "InterchangeTestFunction.h"
+#include "Materials/MaterialInterface.h"
+
+
+UClass* UMaterialImportTestFunctions::GetAssociatedAssetType() const
+{
+	return UMaterialInterface::StaticClass();
+}
+
+
+
+FInterchangeTestFunctionResult UMaterialImportTestFunctions::CheckImportedMaterialCount(const TArray<UMaterialInterface*>& Materials, int32 ExpectedNumberOfImportedMaterials)
+{
+	FInterchangeTestFunctionResult Result;
+	if (Materials.Num() != ExpectedNumberOfImportedMaterials)
+	{
+		Result.AddError(FString::Printf(TEXT("Expected %d materials, imported %d."), ExpectedNumberOfImportedMaterials, Materials.Num()));
+	}
+
+	return Result;
+}
diff --git a/Engine/Plugins/Tests/InterchangeTests/Source/InterchangeTests/Private/ImportTestFunctions/SkeletalMeshImportTestFunctions.cpp b/Engine/Plugins/Tests/InterchangeTests/Source/InterchangeTests/Private/ImportTestFunctions/SkeletalMeshImportTestFunctions.cpp
new file mode 100644
--- /dev/null
+++ b/Engine/Plugins/Tests/InterchangeTests/Source/InterchangeTests/Private/ImportTestFunctions/SkeletalMeshImportTestFunctions.cpp
@@ -0,0 +1,366 @@
+// Copyright Epic Games, Inc. All Rights Reserved.
+
+#include "ImportTestFunctions/SkeletalMeshImportTestFunctions.h"
+#include "InterchangeTestFunction.h"
+#include "Engine/SkeletalMesh.h"
+#include "Rendering/SkeletalMeshModel.h"
+#include "Rendering/SkeletalMeshRenderData.h"
+
+
+UClass* USkeletalMeshImportTestFunctions::GetAssociatedAssetType() const
+{
+	return USkeletalMesh::StaticClass();
+}
+
+
+FInterchangeTestFunctionResult USkeletalMeshImportTestFunctions::CheckImportedSkeletalMeshCount(const TArray<USkeletalMesh*>& Meshes, int32 ExpectedNumberOfImportedSkeletalMeshes)
+{
+	FInterchangeTestFunctionResult Result;
+	if (Meshes.Num() != ExpectedNumberOfImportedSkeletalMeshes)
+	{
+		Result.AddError(FString::Printf(TEXT("Expected %d skeletal meshes, imported %d."), ExpectedNumberOfImportedSkeletalMeshes, Meshes.Num()));
+	}
+
+	return Result;
+}
+
+
+FInterchangeTestFunctionResult USkeletalMeshImportTestFunctions::CheckRenderVertexCount(USkeletalMesh* Mesh, int32 LodIndex, int32 ExpectedNumberOfRenderVertices)
+{
+	FInterchangeTestFunctionResult Result;
+
+	// @todo: add separate test function for checking the source data vertex count?
+	int32 ImportedLods = Mesh->GetResourceForRendering()->LODRenderData.Num();
+	if (LodIndex >= ImportedLods)
+	{
+		Result.AddError(FString::Printf(TEXT("The imported mesh doesn't contain LOD index %d (imported %d)."), LodIndex, ImportedLods));
+	}
+	else
+	{
+		int32 RealVertexCount = Mesh->GetResourceForRendering()->LODRenderData[LodIndex].GetNumVertices();
+		if (RealVertexCount != ExpectedNumberOfRenderVertices)
+		{
+			Result.AddError(FString::Printf(TEXT("For LOD %d, expected %d vertices, imported %d."), LodIndex, ExpectedNumberOfRenderVertices, RealVertexCount));
+		}
+	}
+
+	return Result;
+}
+
+
+FInterchangeTestFunctionResult USkeletalMeshImportTestFunctions::CheckLodCount(USkeletalMesh* Mesh, int32 ExpectedNumberOfLods)
+{
+	FInterchangeTestFunctionResult Result;
+
+	int32 NumLODs = Mesh->GetResourceForRendering()->LODRenderData.Num();
+	if (NumLODs != ExpectedNumberOfLods)
+	{
+		Result.AddError(FString::Printf(TEXT("Expected %d LODs, imported %d."), ExpectedNumberOfLods, NumLODs));
+	}
+
+	return Result;
+}
+
+
+FInterchangeTestFunctionResult USkeletalMeshImportTestFunctions::CheckMaterialSlotCount(USkeletalMesh* Mesh, int32 ExpectedNumberOfMaterialSlots)
+{
+	FInterchangeTestFunctionResult Result;
+
+	int32 NumMaterials = Mesh->GetMaterials().Num();
+	if (NumMaterials != ExpectedNumberOfMaterialSlots)
+	{
+		Result.AddError(FString::Printf(TEXT("Expected %d materials, imported %d."), ExpectedNumberOfMaterialSlots, NumMaterials));
+	}
+
+	return Result;
+}
+
+
+FInterchangeTestFunctionResult USkeletalMeshImportTestFunctions::CheckSectionCount(USkeletalMesh* Mesh, int32 LodIndex, int32 ExpectedNumberOfSections)
+{
+	FInterchangeTestFunctionResult Result;
+
+	int32 ImportedLods = Mesh->GetResourceForRendering()->LODRenderData.Num();
+	if (LodIndex >= ImportedLods)
+	{
+		Result.AddError(FString::Printf(TEXT("The imported mesh doesn't contain LOD index %d (imported %d)."), LodIndex, ImportedLods));
+	}
+	else
+	{
+		int32 NumSections = Mesh->GetResourceForRendering()->LODRenderData[LodIndex].RenderSections.Num();
+		if (NumSections != ExpectedNumberOfSections)
+		{
+			Result.AddError(FString::Printf(TEXT("For LOD %d, expected %d sections, imported %d."), LodIndex, ExpectedNumberOfSections, NumSections));
+		}
+	}
+
+	return Result;
+}
+
+
+FInterchangeTestFunctionResult USkeletalMeshImportTestFunctions::CheckTriangleCountInSection(USkeletalMesh* Mesh, int32 LodIndex, int32 SectionIndex, int32 ExpectedNumberOfTriangles)
+{
+	FInterchangeTestFunctionResult Result;
+
+	int32 ImportedLods = Mesh->GetResourceForRendering()->LODRenderData.Num();
+	if (LodIndex >= ImportedLods)
+	{
+		Result.AddError(FString::Printf(TEXT("The imported mesh doesn't contain LOD index %d (imported %d)."), LodIndex, ImportedLods));
+	}
+	else
+	{
+		int32 NumSections = Mesh->GetResourceForRendering()->LODRenderData[LodIndex].RenderSections.Num();
+		if (SectionIndex >= NumSections)
+		{
+			Result.AddError(FString::Printf(TEXT("The imported mesh doesn't contain section index %d."), SectionIndex));
+		}
+		else
+		{
+			int32 NumberOfTriangles = Mesh->GetResourceForRendering()->LODRenderData[LodIndex].RenderSections[SectionIndex].NumTriangles;
+			if (NumberOfTriangles != ExpectedNumberOfTriangles)
+			{
+				Result.AddError(FString::Printf(TEXT("For LOD %d, section index %d, expected %d triangles, imported %d."), LodIndex, SectionIndex, ExpectedNumberOfTriangles, NumberOfTriangles));
+			}
+		}
+	}
+
+	return Result;
+}
+
+
+FInterchangeTestFunctionResult USkeletalMeshImportTestFunctions::CheckUVChannelCount(USkeletalMesh* Mesh, int32 LodIndex, int32 ExpectedNumberOfUVChannels)
+{
+	FInterchangeTestFunctionResult Result;
+
+	int32 ImportedLods = Mesh->GetResourceForRendering()->LODRenderData.Num();
+	if (LodIndex >= ImportedLods)
+	{
+		Result.AddError(FString::Printf(TEXT("The imported mesh doesn't contain LOD index %d (imported %d)."), LodIndex, ImportedLods));
+	}
+	else
+	{
+		int32 NumUVs = Mesh->GetResourceForRendering()->LODRenderData[LodIndex].GetNumTexCoords();
+		if (NumUVs != ExpectedNumberOfUVChannels)
+		{
+			Result.AddError(FString::Printf(TEXT("For LOD %d, expected %d UVs, imported %d."), LodIndex, ExpectedNumberOfUVChannels, NumUVs));
+		}
+	}
+
+	return Result;
+}
+
+
+FInterchangeTestFunctionResult USkeletalMeshImportTestFunctions::CheckSectionMaterialName(USkeletalMesh* Mesh, int32 LodIndex, int32 SectionIndex, const FString& ExpectedMaterialName)
+{
+	FInterchangeTestFunctionResult Result;
+
+	int32 ImportedLods = Mesh->GetResourceForRendering()->LODRenderData.Num();
+	if (LodIndex >= ImportedLods)
+	{
+		Result.AddError(FString::Printf(TEXT("The imported mesh doesn't contain LOD index %d (imported %d)."), LodIndex, ImportedLods));
+	}
+	else
+	{
+		const TArray<FSkelMeshRenderSection>& Sections = Mesh->GetResourceForRendering()->LODRenderData[LodIndex].RenderSections;
+		if (SectionIndex >= Sections.Num())
+		{
+			Result.AddError(FString::Printf(TEXT("The imported mesh doesn't contain section index %d (imported %d)."), SectionIndex, Sections.Num()));
+		}
+		else
+		{
+			int32 MaterialIndex = Sections[SectionIndex].MaterialIndex;
+
+			const TArray<FSkeletalMaterial>& SkeletalMaterials = Mesh->GetMaterials();
+			if (!SkeletalMaterials.IsValidIndex(MaterialIndex) || SkeletalMaterials[MaterialIndex].MaterialInterface == nullptr)
+			{
+				Result.AddError(FString::Printf(TEXT("The section references a non-existent material (index %d)."), MaterialIndex));
+			}
+			else
+			{
+				FString MaterialName = SkeletalMaterials[MaterialIndex].MaterialInterface->GetName();
+				if (MaterialName != ExpectedMaterialName)
+				{
+					Result.AddError(FString::Printf(TEXT("For LOD %d section %d, expected material name %s, imported %s."), LodIndex, SectionIndex, *ExpectedMaterialName, *MaterialName));
+				}
+			}
+		}
+	}
+
+	return Result;
+}
+
+
+FInterchangeTestFunctionResult USkeletalMeshImportTestFunctions::CheckVertexIndexPosition(USkeletalMesh* Mesh, int32 LodIndex, int32 VertexIndex, const FVector& ExpectedVertexPosition)
+{
+	FInterchangeTestFunctionResult Result;
+
+	int32 ImportedLods = Mesh->GetResourceForRendering()->LODRenderData.Num();
+	if (LodIndex >= ImportedLods)
+	{
+		Result.AddError(FString::Printf(TEXT("The imported mesh doesn't contain LOD index %d (imported %d)."), LodIndex, ImportedLods));
+	}
+	else
+	{
+		int32 VertexCount = Mesh->GetResourceForRendering()->LODRenderData[LodIndex].StaticVertexBuffers.PositionVertexBuffer.GetNumVertices();
+		if (VertexIndex >= VertexCount)
+		{
+			Result.AddError(FString::Printf(TEXT("The imported mesh doesn't contain vertex index %d (imported %d)."), VertexIndex, VertexCount));
+		}
+		else
+		{
+			FVector VertexPosition = Mesh->GetResourceForRendering()->LODRenderData[LodIndex].StaticVertexBuffers.PositionVertexBuffer.VertexPosition(VertexIndex);
+			if (!VertexPosition.Equals(ExpectedVertexPosition))
+			{
+				Result.AddError(FString::Printf(TEXT("For LOD %d vertex index %d, expected position %s, imported %s."), LodIndex, VertexIndex, *ExpectedVertexPosition.ToString(), *VertexPosition.ToString()));
+			}
+		}
+	}
+
+	return Result;
+}
+
+
+FInterchangeTestFunctionResult USkeletalMeshImportTestFunctions::CheckBoneCount(USkeletalMesh* Mesh, int32 ExpectedNumberOfBones)
+{
+	FInterchangeTestFunctionResult Result;
+
+	USkeleton* Skeleton = Mesh->GetSkeleton();
+	if (Skeleton == nullptr)
+	{
+		if (ExpectedNumberOfBones != 0)
+		{
+			Result.AddError(FString::Printf(TEXT("No skeleton found - but expected %d bones"), ExpectedNumberOfBones));
+		}
+	}
+	else
+	{
+		int32 NumberOfBones = Skeleton->GetReferenceSkeleton().GetNum();
+		if (NumberOfBones != ExpectedNumberOfBones)
+		{
+			Result.AddError(FString::Printf(TEXT("Expected %d bones, imported %d."), ExpectedNumberOfBones, NumberOfBones));
+		}
+	}
+
+	return Result;
+}
+
+
+FInterchangeTestFunctionResult USkeletalMeshImportTestFunctions::CheckBonePosition(USkeletalMesh* Mesh, int32 BoneIndex, const FVector& ExpectedBonePosition)
+{
+	FInterchangeTestFunctionResult Result;
+
+	USkeleton* Skeleton = Mesh->GetSkeleton();
+	if (Skeleton == nullptr)
+	{
+		Result.AddError(TEXT("No skeleton found."));
+	}
+	else
+	{
+		int32 NumberOfBones = Skeleton->GetReferenceSkeleton().GetNum();
+		if (BoneIndex >= NumberOfBones)
+		{
+			Result.AddError(FString::Printf(TEXT("Expected bone index %d, but only imported %d bones."), BoneIndex, NumberOfBones));
+		}
+		else
+		{
+			FVector BonePosition = Mesh->GetRefSkeleton().GetRefBonePose()[BoneIndex].GetLocation();
+			if (!BonePosition.Equals(ExpectedBonePosition))
+			{
+				Result.AddError(FString::Printf(TEXT("For bone index %d, expected position %s, imported %s."), BoneIndex, *ExpectedBonePosition.ToString(), *BonePosition.ToString()));
+			}
+		}
+	}
+
+	return Result;
+}
+
+
+FInterchangeTestFunctionResult USkeletalMeshImportTestFunctions::CheckSkinnedVertexCountForBone(USkeletalMesh* Mesh, const FString& BoneName, bool bTestFirstAlternateProfile, int32 ExpectedSkinnedVertexCount)
+{
+	FInterchangeTestFunctionResult Result;
+
+	int32 BoneIndex = Mesh->GetRefSkeleton().FindBoneIndex(*BoneName);
+	if (!Mesh->GetRefSkeleton().IsValidIndex(BoneIndex))
+	{
+		Result.AddError(FString::Printf(TEXT("Could not find bone '%s'."), *BoneName));
+	}
+	else
+	{
+		if (Mesh->GetImportedModel() && Mesh->GetImportedModel()->LODModels.IsValidIndex(0))
+		{
+			int32 SkinnedVerticesForBone = 0;
+			auto IncrementInfluence = [&SkinnedVerticesForBone, &BoneIndex](const FSkelMeshSection& Section,
+																			const FBoneIndexType(&InfluenceBones)[MAX_TOTAL_INFLUENCES],
+																			const uint8(&InfluenceWeights)[MAX_TOTAL_INFLUENCES])
+			{
+				for (int32 InfluenceIndex = 0; InfluenceIndex < MAX_TOTAL_INFLUENCES; ++InfluenceIndex)
+				{
+					if (InfluenceWeights[InfluenceIndex] == 0)
+					{
+						// Influences are sorted by weight so no need to go further then a zero weight
+						break;
+					}
+					if (Section.BoneMap[InfluenceBones[InfluenceIndex]] == BoneIndex)
+					{
+						SkinnedVerticesForBone++;
+						break;
+					}
+				}
+			};
+
+			if (!bTestFirstAlternateProfile)
+			{
+				for (const FSkelMeshSection& Section : Mesh->GetImportedModel()->LODModels[0].Sections)
+				{
+					const int32 SectionVertexCount = Section.SoftVertices.Num();
+					// Find the number of vertices skinned to this bone
+					for (int32 SectionVertexIndex = 0; SectionVertexIndex < SectionVertexCount; ++SectionVertexIndex)
+					{
+						const FSoftSkinVertex& Vertex = Section.SoftVertices[SectionVertexIndex];
+						IncrementInfluence(Section, Vertex.InfluenceBones, Vertex.InfluenceWeights);
+					}
+				}
+			}
+			else
+			{
+				if (Mesh->GetSkinWeightProfiles().Num() > 0)
+				{
+					const int32 TotalVertexCount = Mesh->GetImportedModel()->LODModels[0].NumVertices;
+					const FSkinWeightProfileInfo& SkinWeightProfile = Mesh->GetSkinWeightProfiles()[0];
+					const FImportedSkinWeightProfileData& SkinWeightData = Mesh->GetImportedModel()->LODModels[0].SkinWeightProfiles.FindChecked(SkinWeightProfile.Name);
+
+					if (SkinWeightData.SkinWeights.Num() != TotalVertexCount)
+					{
+						Result.AddError(TEXT("Unable to find alternate skinning profile, please uncheck the 'test alternate profile' box."));
+					}
+					else
+					{
+						int32 TotalVertexIndex = 0;
+						for (const FSkelMeshSection& Section : Mesh->GetImportedModel()->LODModels[0].Sections)
+						{
+							const int32 SectionVertexCount = Section.SoftVertices.Num();
+							//Find the number of vertex skin by this bone
+							for (int32 SectionVertexIndex = 0; SectionVertexIndex < SectionVertexCount; ++SectionVertexIndex, ++TotalVertexIndex)
+							{
+								const FRawSkinWeight& SkinWeight = SkinWeightData.SkinWeights[TotalVertexIndex];
+								IncrementInfluence(Section, SkinWeight.InfluenceBones, SkinWeight.InfluenceWeights);
+							}
+						}
+					}
+				}
+			}
+
+			if (SkinnedVerticesForBone != ExpectedSkinnedVertexCount)
+			{
+				Result.AddError(FString::Printf(TEXT("For bone '%s', expected %d vertices, imported %s."), *BoneName, ExpectedSkinnedVertexCount, SkinnedVerticesForBone));
+			}
+		}
+		else
+		{
+			Result.AddError(TEXT("No valid mesh geometry found to find the vertex count"));
+		}
+	}
+
+	return Result;
+}
+
diff --git a/Engine/Plugins/Tests/InterchangeTests/Source/InterchangeTests/Private/ImportTestFunctions/StaticMeshImportTestFunctions.cpp b/Engine/Plugins/Tests/InterchangeTests/Source/InterchangeTests/Private/ImportTestFunctions/StaticMeshImportTestFunctions.cpp
new file mode 100644
--- /dev/null
+++ b/Engine/Plugins/Tests/InterchangeTests/Source/InterchangeTests/Private/ImportTestFunctions/StaticMeshImportTestFunctions.cpp
@@ -0,0 +1,646 @@
+// Copyright Epic Games, Inc. All Rights Reserved.
+
+#include "ImportTestFunctions/StaticMeshImportTestFunctions.h"
+#include "InterchangeTestFunction.h"
+#include "Engine/StaticMesh.h"
+#include "Engine/SkeletalMesh.h"
+#include "StaticMeshAttributes.h"
+
+
+UClass* UStaticMeshImportTestFunctions::GetAssociatedAssetType() const
+{
+	return UStaticMesh::StaticClass();
+}
+
+
+FInterchangeTestFunctionResult UStaticMeshImportTestFunctions::CheckImportedStaticMeshCount(const TArray<UStaticMesh*>& Meshes, int32 ExpectedNumberOfImportedStaticMeshes)
+{
+	FInterchangeTestFunctionResult Result;
+	if (Meshes.Num() != ExpectedNumberOfImportedStaticMeshes)
+	{
+		Result.AddError(FString::Printf(TEXT("Expected %d static meshes, imported %d."), ExpectedNumberOfImportedStaticMeshes, Meshes.Num()));
+	}
+
+	return Result;
+}
+
+
+FInterchangeTestFunctionResult UStaticMeshImportTestFunctions::CheckVertexCount(UStaticMesh* Mesh, int32 LodIndex, int32 ExpectedNumberOfVertices)
+{
+	FInterchangeTestFunctionResult Result;
+
+	const FMeshDescription* MeshDescription = Mesh->GetMeshDescription(LodIndex);
+	if (MeshDescription == nullptr)
+	{
+		Result.AddError(FString::Printf(TEXT("The imported mesh doesn't contain LOD index %d."), LodIndex));
+	}
+	else
+	{
+		int32 RealVertexCount = MeshDescription->Vertices().Num();
+		if (RealVertexCount != ExpectedNumberOfVertices)
+		{
+			Result.AddError(FString::Printf(TEXT("For LOD %d, expected %d vertices, imported %d."), LodIndex, ExpectedNumberOfVertices, RealVertexCount));
+		}
+	}
+
+	return Result;
+}
+
+
+FInterchangeTestFunctionResult UStaticMeshImportTestFunctions::CheckRenderVertexCount(UStaticMesh* Mesh, int32 LodIndex, int32 ExpectedNumberOfRenderVertices)
+{
+	FInterchangeTestFunctionResult Result;
+
+	if (LodIndex >= Mesh->GetNumLODs())
+	{
+		Result.AddError(FString::Printf(TEXT("The imported mesh doesn't contain LOD index %d."), LodIndex));
+	}
+	else
+	{
+		int32 RealVertexCount = Mesh->GetNumVertices(LodIndex);
+		if (RealVertexCount != ExpectedNumberOfRenderVertices)
+		{
+			Result.AddError(FString::Printf(TEXT("For LOD %d, expected %d vertices, imported %d."), LodIndex, ExpectedNumberOfRenderVertices, RealVertexCount));
+		}
+	}
+
+	return Result;
+}
+
+
+FInterchangeTestFunctionResult UStaticMeshImportTestFunctions::CheckLodCount(UStaticMesh* Mesh, int32 ExpectedNumberOfLods)
+{
+	FInterchangeTestFunctionResult Result;
+
+	int32 NumLODs = Mesh->GetNumSourceModels();
+	if (NumLODs != ExpectedNumberOfLods)
+	{
+		Result.AddError(FString::Printf(TEXT("Expected %d LODs, imported %d."), ExpectedNumberOfLods, NumLODs));
+	}
+
+	return Result;
+}
+
+
+FInterchangeTestFunctionResult UStaticMeshImportTestFunctions::CheckMaterialSlotCount(UStaticMesh* Mesh, int32 ExpectedNumberOfMaterialSlots)
+{
+	FInterchangeTestFunctionResult Result;
+
+	int32 NumMaterials = Mesh->GetStaticMaterials().Num();
+	if (NumMaterials != ExpectedNumberOfMaterialSlots)
+	{
+		Result.AddError(FString::Printf(TEXT("Expected %d materials, imported %d."), ExpectedNumberOfMaterialSlots, NumMaterials));
+	}
+
+	return Result;
+}
+
+
+FInterchangeTestFunctionResult UStaticMeshImportTestFunctions::CheckPolygonGroupCount(UStaticMesh* Mesh, int32 LodIndex, int32 ExpectedNumberOfPolygonGroups)
+{
+	FInterchangeTestFunctionResult Result;
+
+	const FMeshDescription* MeshDescription = Mesh->GetMeshDescription(LodIndex);
+	if (MeshDescription == nullptr)
+	{
+		Result.AddError(FString::Printf(TEXT("The imported mesh doesn't contain LOD index %d."), LodIndex));
+	}
+	else
+	{
+		int32 RealNumberOfPolygonGroups = MeshDescription->PolygonGroups().Num();
+		if (RealNumberOfPolygonGroups != ExpectedNumberOfPolygonGroups)
+		{
+			Result.AddError(FString::Printf(TEXT("For LOD %d, expected %d polygon groups, imported %d."), LodIndex, ExpectedNumberOfPolygonGroups, RealNumberOfPolygonGroups));
+		}
+	}
+
+	return Result;
+}
+
+
+FInterchangeTestFunctionResult UStaticMeshImportTestFunctions::CheckSectionCount(UStaticMesh* Mesh, int32 LodIndex, int32 ExpectedNumberOfSections)
+{
+	FInterchangeTestFunctionResult Result;
+
+	if (LodIndex >= Mesh->GetNumLODs())
+	{
+		Result.AddError(FString::Printf(TEXT("The imported mesh doesn't contain LOD index %d."), LodIndex));
+	}
+	else
+	{
+		int32 NumSections = Mesh->GetNumSections(LodIndex);
+		if (NumSections != ExpectedNumberOfSections)
+		{
+			Result.AddError(FString::Printf(TEXT("For LOD %d, expected %d sections, imported %d."), LodIndex, ExpectedNumberOfSections, NumSections));
+		}
+
+	}
+
+	return Result;
+}
+
+
+FInterchangeTestFunctionResult UStaticMeshImportTestFunctions::CheckTotalTriangleCount(UStaticMesh* Mesh, int32 LodIndex, int32 ExpectedTotalNumberOfTriangles)
+{
+	FInterchangeTestFunctionResult Result;
+
+	const FMeshDescription* MeshDescription = Mesh->GetMeshDescription(LodIndex);
+	if (MeshDescription == nullptr)
+	{
+		Result.AddError(FString::Printf(TEXT("The imported mesh doesn't contain LOD index %d."), LodIndex));
+	}
+	else
+	{
+		int32 RealNumberOfTriangles = MeshDescription->Triangles().Num();
+		if (RealNumberOfTriangles != ExpectedTotalNumberOfTriangles)
+		{
+			Result.AddError(FString::Printf(TEXT("For LOD %d, expected %d triangles, imported %d."), LodIndex, ExpectedTotalNumberOfTriangles, RealNumberOfTriangles));
+		}
+	}
+
+	return Result;
+}
+
+
+FInterchangeTestFunctionResult UStaticMeshImportTestFunctions::CheckTotalPolygonCount(UStaticMesh* Mesh, int32 LodIndex, int32 ExpectedTotalNumberOfPolygons)
+{
+	FInterchangeTestFunctionResult Result;
+
+	const FMeshDescription* MeshDescription = Mesh->GetMeshDescription(LodIndex);
+	if (MeshDescription == nullptr)
+	{
+		Result.AddError(FString::Printf(TEXT("The imported mesh doesn't contain LOD index %d."), LodIndex));
+	}
+	else
+	{
+		int32 RealNumberOfPolygons = MeshDescription->Polygons().Num();
+		if (RealNumberOfPolygons != ExpectedTotalNumberOfPolygons)
+		{
+			Result.AddError(FString::Printf(TEXT("For LOD %d, expected %d polygons, imported %d."), LodIndex, ExpectedTotalNumberOfPolygons, RealNumberOfPolygons));
+		}
+	}
+
+	return Result;
+}
+
+
+FInterchangeTestFunctionResult UStaticMeshImportTestFunctions::CheckThatMeshHasQuadsOrNgons(UStaticMesh* Mesh, int32 LodIndex, bool bMeshHasQuadsOrNgons)
+{
+	FInterchangeTestFunctionResult Result;
+
+	const FMeshDescription* MeshDescription = Mesh->GetMeshDescription(LodIndex);
+	if (MeshDescription == nullptr)
+	{
+		Result.AddError(FString::Printf(TEXT("The imported mesh doesn't contain LOD index %d."), LodIndex));
+	}
+	else
+	{
+		int32 NumberOfTriangles = MeshDescription->Triangles().Num();
+		int32 NumberOfPolygons = MeshDescription->Polygons().Num();
+		bool bReallyHasQuads = (NumberOfTriangles > NumberOfPolygons);
+
+		if (bReallyHasQuads != bMeshHasQuadsOrNgons)
+		{
+			Result.AddError(FString::Printf(TEXT("For LOD %d, expected has quads %d, imported %d."), LodIndex, bMeshHasQuadsOrNgons, bReallyHasQuads));
+		}
+	}
+
+	return Result;
+}
+
+
+FInterchangeTestFunctionResult UStaticMeshImportTestFunctions::CheckTriangleCountInPolygonGroup(UStaticMesh* Mesh, int32 LodIndex, int32 PolygonGroupIndex, int32 ExpectedNumberOfTriangles)
+{
+	FInterchangeTestFunctionResult Result;
+
+	const FMeshDescription* MeshDescription = Mesh->GetMeshDescription(LodIndex);
+	if (MeshDescription == nullptr)
+	{
+		Result.AddError(FString::Printf(TEXT("The imported mesh doesn't contain LOD index %d."), LodIndex));
+	}
+	else
+	{
+		int32 RealNumberOfTriangles = MeshDescription->GetNumPolygonGroupTriangles(PolygonGroupIndex);
+		if (RealNumberOfTriangles != ExpectedNumberOfTriangles)
+		{
+			Result.AddError(FString::Printf(TEXT("For LOD %d PolygonGroup %d, expected %d triangles, imported %d."), LodIndex, PolygonGroupIndex, ExpectedNumberOfTriangles, RealNumberOfTriangles));
+		}
+	}
+
+	return Result;
+}
+
+
+FInterchangeTestFunctionResult UStaticMeshImportTestFunctions::CheckPolygonCountInPolygonGroup(UStaticMesh* Mesh, int32 LodIndex, int32 PolygonGroupIndex, int32 ExpectedNumberOfPolygons)
+{
+	FInterchangeTestFunctionResult Result;
+
+	const FMeshDescription* MeshDescription = Mesh->GetMeshDescription(LodIndex);
+	if (MeshDescription == nullptr)
+	{
+		Result.AddError(FString::Printf(TEXT("The imported mesh doesn't contain LOD index %d."), LodIndex));
+	}
+	else
+	{
+		int32 RealNumberOfPolygons = MeshDescription->GetNumPolygonGroupPolygons(PolygonGroupIndex);
+		if (RealNumberOfPolygons != ExpectedNumberOfPolygons)
+		{
+			Result.AddError(FString::Printf(TEXT("For LOD %d PolygonGroup %d, expected %d polygons, imported %d."), LodIndex, PolygonGroupIndex, ExpectedNumberOfPolygons, RealNumberOfPolygons));
+		}
+	}
+
+	return Result;
+}
+
+
+FInterchangeTestFunctionResult UStaticMeshImportTestFunctions::CheckUVChannelCount(UStaticMesh* Mesh, int32 LodIndex, int32 ExpectedNumberOfUVChannels)
+{
+	FInterchangeTestFunctionResult Result;
+
+	const FMeshDescription* MeshDescription = Mesh->GetMeshDescription(LodIndex);
+	if (MeshDescription == nullptr)
+	{
+		Result.AddError(FString::Printf(TEXT("The imported mesh doesn't contain LOD index %d."), LodIndex));
+	}
+	else
+	{
+		FStaticMeshConstAttributes Attributes(*MeshDescription);
+
+		int32 NumberOfUVChannels = Attributes.GetVertexInstanceUVs().GetNumChannels();
+		if (NumberOfUVChannels != ExpectedNumberOfUVChannels)
+		{
+			Result.AddError(FString::Printf(TEXT("Expected %d UV Channels, imported %d."), ExpectedNumberOfUVChannels, NumberOfUVChannels));
+		}
+	}
+
+	return Result;
+}
+
+
+FInterchangeTestFunctionResult UStaticMeshImportTestFunctions::CheckPolygonGroupImportedMaterialSlotName(UStaticMesh* Mesh, int32 LodIndex, int32 PolygonGroupIndex, const FString& ExpectedImportedMaterialSlotName)
+{
+	FInterchangeTestFunctionResult Result;
+
+	const FMeshDescription* MeshDescription = Mesh->GetMeshDescription(LodIndex);
+	if (MeshDescription == nullptr)
+	{
+		Result.AddError(FString::Printf(TEXT("The imported mesh doesn't contain LOD index %d."), LodIndex));
+	}
+	else
+	{
+		FStaticMeshConstAttributes Attributes(*MeshDescription);
+
+		FString MaterialSlotName = Attributes.GetPolygonGroupMaterialSlotNames()[PolygonGroupIndex].ToString();
+		if (MaterialSlotName != ExpectedImportedMaterialSlotName)
+		{
+			Result.AddError(FString::Printf(TEXT("For LOD %d PolygonGroup %d, expected slot name %s, imported %s."), LodIndex, PolygonGroupIndex, *ExpectedImportedMaterialSlotName, *MaterialSlotName));
+		}
+	}
+
+	return Result;
+}
+
+
+FInterchangeTestFunctionResult UStaticMeshImportTestFunctions::CheckSectionMaterialIndex(UStaticMesh* Mesh, int32 LodIndex, int32 SectionIndex, int32 ExpectedMaterialIndex)
+{
+	FInterchangeTestFunctionResult Result;
+
+	if (LodIndex >= Mesh->GetNumLODs())
+	{
+		Result.AddError(FString::Printf(TEXT("The imported mesh doesn't contain LOD index %d."), LodIndex));
+	}
+	else
+	{
+		const FStaticMeshSectionArray& Sections = Mesh->GetRenderData()->LODResources[LodIndex].Sections;
+		if (SectionIndex >= Sections.Num())
+		{
+			Result.AddError(FString::Printf(TEXT("The imported mesh doesn't contain section index %d."), SectionIndex));
+		}
+		else
+		{
+			int32 MaterialIndex = Sections[SectionIndex].MaterialIndex;
+
+			if (MaterialIndex != ExpectedMaterialIndex)
+			{
+				Result.AddError(FString::Printf(TEXT("For LOD %d section %d, expected material index %d, imported %d."), LodIndex, SectionIndex, ExpectedMaterialIndex, MaterialIndex));
+			}
+		}
+	}
+
+	return Result;
+}
+
+
+FInterchangeTestFunctionResult UStaticMeshImportTestFunctions::CheckSectionMaterialName(UStaticMesh* Mesh, int32 LodIndex, int32 SectionIndex, const FString& ExpectedMaterialName)
+{
+	FInterchangeTestFunctionResult Result;
+
+	if (LodIndex >= Mesh->GetNumLODs())
+	{
+		Result.AddError(FString::Printf(TEXT("The imported mesh doesn't contain LOD index %d."), LodIndex));
+	}
+	else
+	{
+		const FStaticMeshSectionArray& Sections = Mesh->GetRenderData()->LODResources[LodIndex].Sections;
+		if (SectionIndex >= Sections.Num())
+		{
+			Result.AddError(FString::Printf(TEXT("The imported mesh doesn't contain section index %d."), SectionIndex));
+		}
+		else
+		{
+			int32 MaterialIndex = Sections[SectionIndex].MaterialIndex;
+
+			const TArray<FStaticMaterial>& StaticMaterials = Mesh->GetStaticMaterials();
+			if (!StaticMaterials.IsValidIndex(MaterialIndex) || StaticMaterials[MaterialIndex].MaterialInterface == nullptr)
+			{
+				Result.AddError(FString::Printf(TEXT("The section references a non-existent material (index %d)."), MaterialIndex));
+			}
+			else
+			{
+				FString MaterialName = StaticMaterials[MaterialIndex].MaterialInterface->GetName();
+				if (MaterialName != ExpectedMaterialName)
+				{
+					Result.AddError(FString::Printf(TEXT("For LOD %d section %d, expected material name %s, imported %s."), LodIndex, SectionIndex, *ExpectedMaterialName, *MaterialName));
+				}
+			}
+		}
+	}
+
+	return Result;
+}
+
+
+FInterchangeTestFunctionResult UStaticMeshImportTestFunctions::CheckVertexIndexPosition(UStaticMesh* Mesh, int32 LodIndex, int32 VertexIndex, const FVector& ExpectedVertexPosition)
+{
+	FInterchangeTestFunctionResult Result;
+
+	const FMeshDescription* MeshDescription = Mesh->GetMeshDescription(LodIndex);
+	if (MeshDescription == nullptr)
+	{
+		Result.AddError(FString::Printf(TEXT("The imported mesh doesn't contain LOD index %d."), LodIndex));
+	}
+	else
+	{
+		FStaticMeshConstAttributes Attributes(*MeshDescription);
+		
+		if (VertexIndex >= Attributes.GetVertexPositions().GetNumElements())
+		{
+			Result.AddError(FString::Printf(TEXT("The imported mesh doesn't contain a vertex of index %d."), VertexIndex));
+		}
+		else
+		{
+			const FVector& VertexPosition = Attributes.GetVertexPositions()[VertexIndex];
+			if (!VertexPosition.Equals(ExpectedVertexPosition))
+			{
+				Result.AddError(FString::Printf(TEXT("For LOD %d vertex index %d, expected position %s, imported %s."), LodIndex, VertexIndex, *ExpectedVertexPosition.ToString(), *VertexPosition.ToString()));
+			}
+		}
+	}
+
+	return Result;
+}
+
+
+FInterchangeTestFunctionResult UStaticMeshImportTestFunctions::CheckSimpleCollisionPrimitiveCount(UStaticMesh* Mesh, int32 ExpectedSphereElementCount, int32 ExpectedBoxElementCount, int32 ExpectedCapsuleElementCount, int32 ExpectedConvexElementCount, int32 ExpectedTaperedCapsuleElementCount)
+{
+	FInterchangeTestFunctionResult Result;
+
+	UBodySetup* BodySetup = Mesh->GetBodySetup();
+	if (BodySetup == nullptr)
+	{
+		Result.AddError(TEXT("Didn't find a BodySetup on the mesh"));
+	}
+	else
+	{
+		int32 SphereElementCount = BodySetup->AggGeom.SphereElems.Num();
+		int32 BoxElementCount = BodySetup->AggGeom.BoxElems.Num();
+		int32 CapsuleElementCount = BodySetup->AggGeom.SphylElems.Num();
+		int32 ConvexElementCount = BodySetup->AggGeom.ConvexElems.Num();
+		int32 TaperedCapsuleElementCount = BodySetup->AggGeom.TaperedCapsuleElems.Num();
+
+		if (SphereElementCount != ExpectedSphereElementCount)
+		{
+			Result.AddError(FString::Printf(TEXT("Expected %d sphere elements, imported %d."), ExpectedSphereElementCount, SphereElementCount));
+		}
+
+		if (BoxElementCount != ExpectedBoxElementCount)
+		{
+			Result.AddError(FString::Printf(TEXT("Expected %d box elements, imported %d."), ExpectedBoxElementCount, BoxElementCount));
+		}
+
+		if (CapsuleElementCount != ExpectedCapsuleElementCount)
+		{
+			Result.AddError(FString::Printf(TEXT("Expected %d capsule elements, imported %d."), ExpectedCapsuleElementCount, CapsuleElementCount));
+		}
+
+		if (ConvexElementCount != ExpectedConvexElementCount)
+		{
+			Result.AddError(FString::Printf(TEXT("Expected %d convex elements, imported %d."), ExpectedConvexElementCount, ConvexElementCount));
+		}
+
+		if (TaperedCapsuleElementCount != ExpectedTaperedCapsuleElementCount)
+		{
+			Result.AddError(FString::Printf(TEXT("Expected %d tapered capsule elements, imported %d."), ExpectedTaperedCapsuleElementCount, TaperedCapsuleElementCount));
+		}
+	}
+
+
+	return Result;
+}
+
+
+FInterchangeTestFunctionResult UStaticMeshImportTestFunctions::CheckAgainstGroundTruth(UStaticMesh* Mesh, TSoftObjectPtr<UStaticMesh> MeshToCompare,
+	bool bCheckVertexCountEqual,
+	bool bCheckTriangleCountEqual,
+	bool bCheckUVChannelCountEqual,
+	bool bCheckCollisionPrimitiveCountEqual,
+	bool bCheckVertexPositionsEqual,
+	bool bCheckNormalsEqual
+)
+{
+	FInterchangeTestFunctionResult Result;
+
+	UStaticMesh* GroundTruth = MeshToCompare.LoadSynchronous();
+	if (GroundTruth == nullptr)
+	{
+		Result.AddError(FString::Printf(TEXT("Unable to find ground truth mesh '%s'"), *MeshToCompare.ToString()));
+		return Result;
+	}
+
+	int32 NumLods = Mesh->GetNumSourceModels();
+	int32 ExpectedNumLods = GroundTruth->GetNumSourceModels();
+
+	if (NumLods != ExpectedNumLods)
+	{
+		Result.AddError(FString::Printf(TEXT("Number of imported LODs don't match: expected %d, imported %d"), ExpectedNumLods, NumLods));
+		return Result;
+	}
+
+	if (bCheckVertexCountEqual)
+	{
+		for (int32 LodIndex = 0; LodIndex < NumLods; LodIndex++)
+		{
+			const FMeshDescription* MD = Mesh->GetMeshDescription(LodIndex);
+			const FMeshDescription* GroundTruthMD = GroundTruth->GetMeshDescription(LodIndex);
+
+			int32 VertexCount = MD->Vertices().Num();
+			int32 ExpectedVertexCount = GroundTruthMD->Vertices().Num();
+
+			if (VertexCount != ExpectedVertexCount)
+			{
+				Result.AddError(FString::Printf(TEXT("Number of vertices don't match in LOD %d: expected %d, imported %d"), LodIndex, ExpectedVertexCount, VertexCount));
+			}
+		}
+	}
+
+	if (bCheckTriangleCountEqual)
+	{
+		for (int32 LodIndex = 0; LodIndex < NumLods; LodIndex++)
+		{
+			const FMeshDescription* MD = Mesh->GetMeshDescription(LodIndex);
+			const FMeshDescription* GroundTruthMD = GroundTruth->GetMeshDescription(LodIndex);
+
+			int32 TriangleCount = MD->Triangles().Num();
+			int32 ExpectedTriangleCount = GroundTruthMD->Triangles().Num();
+
+			if (TriangleCount != ExpectedTriangleCount)
+			{
+				Result.AddError(FString::Printf(TEXT("Number of triangles don't match in LOD %d: expected %d, imported %d"), LodIndex, ExpectedTriangleCount, TriangleCount));
+			}
+		}
+	}
+
+	if (bCheckUVChannelCountEqual)
+	{
+		for (int32 LodIndex = 0; LodIndex < NumLods; LodIndex++)
+		{
+			const FMeshDescription* MD = Mesh->GetMeshDescription(LodIndex);
+			const FMeshDescription* GroundTruthMD = GroundTruth->GetMeshDescription(LodIndex);
+
+			FStaticMeshConstAttributes Attributes(*MD);
+			FStaticMeshConstAttributes GroundTruthAttributes(*GroundTruthMD);
+
+			int32 UVChannels = Attributes.GetVertexInstanceUVs().GetNumChannels();
+			int32 ExpectedUVChannels = GroundTruthAttributes.GetVertexInstanceUVs().GetNumChannels();
+
+			if (UVChannels != ExpectedUVChannels)
+			{
+				Result.AddError(FString::Printf(TEXT("Number of UV channels don't match in LOD %d: expected %d, imported %d"), LodIndex, ExpectedUVChannels, UVChannels));
+			}
+		}
+	}
+
+	if (bCheckCollisionPrimitiveCountEqual)
+	{
+		UBodySetup* BodySetup = Mesh->GetBodySetup();
+		UBodySetup* GroundTruthBodySetup = GroundTruth->GetBodySetup();
+
+		if (BodySetup == nullptr || GroundTruthBodySetup == nullptr)
+		{
+			Result.AddError(TEXT("No collision data found to compare"));
+		}
+		else
+		{
+			const FKAggregateGeom& AggGeom = BodySetup->AggGeom;
+			const FKAggregateGeom& GroundTruthAggGeom = GroundTruthBodySetup->AggGeom;
+
+			if (AggGeom.BoxElems.Num() != GroundTruthAggGeom.BoxElems.Num())
+			{
+				Result.AddError(FString::Printf(TEXT("Box collision count mismatch: expected %d, imported %d"), GroundTruthAggGeom.BoxElems.Num(), AggGeom.BoxElems.Num()));
+			}
+
+			if (AggGeom.SphereElems.Num() != GroundTruthAggGeom.SphereElems.Num())
+			{
+				Result.AddError(FString::Printf(TEXT("Sphere collision count mismatch: expected %d, imported %d"), GroundTruthAggGeom.SphereElems.Num(), AggGeom.SphereElems.Num()));
+			}
+
+			if (AggGeom.ConvexElems.Num() != GroundTruthAggGeom.ConvexElems.Num())
+			{
+				Result.AddError(FString::Printf(TEXT("Convex collision count mismatch: expected %d, imported %d"), GroundTruthAggGeom.ConvexElems.Num(), AggGeom.ConvexElems.Num()));
+			}
+
+			if (AggGeom.SphylElems.Num() != GroundTruthAggGeom.SphylElems.Num())
+			{
+				Result.AddError(FString::Printf(TEXT("Capsule collision count mismatch: expected %d, imported %d"), GroundTruthAggGeom.SphylElems.Num(), AggGeom.SphylElems.Num()));
+			}
+
+			if (AggGeom.TaperedCapsuleElems.Num() != GroundTruthAggGeom.TaperedCapsuleElems.Num())
+			{
+				Result.AddError(FString::Printf(TEXT("Tapered capsule collision count mismatch: expected %d, imported %d"), GroundTruthAggGeom.TaperedCapsuleElems.Num(), AggGeom.TaperedCapsuleElems.Num()));
+			}
+		}
+	}
+
+	if (bCheckVertexPositionsEqual)
+	{
+		for (int32 LodIndex = 0; LodIndex < NumLods; LodIndex++)
+		{
+			const FMeshDescription* MD = Mesh->GetMeshDescription(LodIndex);
+			const FMeshDescription* GroundTruthMD = GroundTruth->GetMeshDescription(LodIndex);
+
+			FStaticMeshConstAttributes Attributes(*MD);
+			FStaticMeshConstAttributes GroundTruthAttributes(*GroundTruthMD);
+
+			int32 VertexCount = Attributes.GetVertexPositions().GetNumElements();
+			int32 ExpectedVertexCount = GroundTruthAttributes.GetVertexPositions().GetNumElements();
+
+			if (VertexCount != ExpectedVertexCount)
+			{
+				Result.AddError(FString::Printf(TEXT("Number of vertices don't match in LOD %d: expected %d, imported %d"), LodIndex, ExpectedVertexCount, VertexCount));
+			}
+			else
+			{
+				for (int32 VertexIndex = 0; VertexIndex < VertexCount; VertexIndex++)
+				{
+					const FVector& V0 = Attributes.GetVertexPositions()[VertexIndex];
+					const FVector& V1 = GroundTruthAttributes.GetVertexPositions()[VertexIndex];
+
+					if (!V0.Equals(V1))
+					{
+						Result.AddError(FString::Printf(TEXT("LOD %d: vertex index %d: position mismatch: expected %s, imported %s"), LodIndex, VertexIndex, *V1.ToString(), *V0.ToString()));
+					}
+				}
+			}
+		}
+	}
+
+	if (bCheckNormalsEqual)
+	{
+		for (int32 LodIndex = 0; LodIndex < NumLods; LodIndex++)
+		{
+			const FMeshDescription* MD = Mesh->GetMeshDescription(LodIndex);
+			const FMeshDescription* GroundTruthMD = GroundTruth->GetMeshDescription(LodIndex);
+
+			FStaticMeshConstAttributes Attributes(*MD);
+			FStaticMeshConstAttributes GroundTruthAttributes(*GroundTruthMD);
+
+			int32 NormalsCount = Attributes.GetVertexInstanceNormals().GetNumElements();
+			int32 ExpectedNormalsCount = GroundTruthAttributes.GetVertexInstanceNormals().GetNumElements();
+
+			if (NormalsCount != ExpectedNormalsCount)
+			{
+				Result.AddError(FString::Printf(TEXT("Number of vertices don't match in LOD %d: expected %d, imported %d"), LodIndex, ExpectedNormalsCount, NormalsCount));
+			}
+			else
+			{
+				for (int32 NormalIndex = 0; NormalIndex < NormalsCount; NormalIndex++)
+				{
+					const FVector& N0 = Attributes.GetVertexInstanceNormals()[NormalIndex];
+					const FVector& N1 = GroundTruthAttributes.GetVertexInstanceNormals()[NormalIndex];
+
+					if (!N0.Equals(N1))
+					{
+						Result.AddError(FString::Printf(TEXT("LOD %d: normal index %d mismatch: expected %s, imported %s"), LodIndex, NormalIndex, *N1.ToString(), *N0.ToString()));
+					}
+				}
+			}
+		}
+	}
+
+	// Don't keep the ground truth asset lying around; mark it as garbage to get rid of it at the end of the tests
+	GroundTruth->RemoveFromRoot();
+	GroundTruth->MarkAsGarbage();
+
+	return Result;
+}
+
diff --git a/Engine/Plugins/Tests/InterchangeTests/Source/InterchangeTests/Private/ImportTestFunctions/TextureImportTestFunctions.cpp b/Engine/Plugins/Tests/InterchangeTests/Source/InterchangeTests/Private/ImportTestFunctions/TextureImportTestFunctions.cpp
new file mode 100644
--- /dev/null
+++ b/Engine/Plugins/Tests/InterchangeTests/Source/InterchangeTests/Private/ImportTestFunctions/TextureImportTestFunctions.cpp
@@ -0,0 +1,11 @@
+// Copyright Epic Games, Inc. All Rights Reserved.
+
+#include "ImportTestFunctions/TextureImportTestFunctions.h"
+#include "InterchangeTestFunction.h"
+#include "Engine/Texture.h"
+
+
+UClass* UTextureImportTestFunctions::GetAssociatedAssetType() const
+{
+	return UTexture::StaticClass();
+}
diff --git a/Engine/Plugins/Tests/InterchangeTests/Source/InterchangeTests/Private/InterchangeImportTest.cpp b/Engine/Plugins/Tests/InterchangeTests/Source/InterchangeTests/Private/InterchangeImportTest.cpp
new file mode 100644
--- /dev/null
+++ b/Engine/Plugins/Tests/InterchangeTests/Source/InterchangeTests/Private/InterchangeImportTest.cpp
@@ -0,0 +1,227 @@
+// Copyright Epic Games, Inc. All Rights Reserved.
+
+#include "CoreMinimal.h"
+#include "HAL/PlatformFileManager.h"
+#include "HAL/FileManager.h"
+#include "Misc/Paths.h"
+#include "Misc/WildcardString.h"
+#include "Misc/AutomationTest.h"
+#include "InterchangeTestsLog.h"
+#include "ImportTestFunctions/StaticMeshImportTestFunctions.h"
+#include "InterchangeImportTestData.h"
+#include "InterchangeImportTestSettings.h"
+#include "InterchangeImportTestPlan.h"
+#include "InterchangeImportTestStepBase.h"
+#include "InterchangeTestsModule.h"
+#include "UObject/StrongObjectPtr.h"
+#include "AssetRegistryModule.h"
+#include "AssetData.h"
+#include "InterchangeManager.h"
+#include "UObject/SavePackage.h"
+#include "ObjectTools.h"
+
+#include "Engine/StaticMesh.h"
+#include "InterchangeGenericAssetsPipeline.h"
+
+
+IMPLEMENT_COMPLEX_AUTOMATION_TEST(FInterchangeImportTest, "Editor.Interchange.Import", EAutomationTestFlags::EditorContext | EAutomationTestFlags::EngineFilter)
+
+
+void FInterchangeImportTest::GetTests(TArray<FString>& OutBeautifiedNames, TArray<FString>& OutTestCommands) const
+{
+	const FAssetRegistryModule& AssetRegistryModule = FModuleManager::LoadModuleChecked<FAssetRegistryModule>(TEXT("AssetRegistry"));
+	TArray<FAssetData> AllTestPlans;
+	AssetRegistryModule.Get().GetAssetsByClass(UInterchangeImportTestPlan::StaticClass()->GetFName(), AllTestPlans, true);
+
+	// Get a list of all paths containing InterchangeTestPlan assets
+	TSet<FName> Paths;
+	for (const FAssetData& TestPlan : AllTestPlans)
+	{
+		Paths.Add(TestPlan.PackagePath);
+	}
+
+	// And create a sorted list from the set.
+	// Each unique path will be a sub-entry in the automated test list.
+	// All test plans within a given folder will be executed in parallel.
+
+	Paths.Sort([](const FName& A, const FName& B) { return A.ToString() < B.ToString(); });
+
+	for (const FName& Path : Paths)
+	{
+		FString PathAsString = Path.ToString();
+		OutTestCommands.Add(PathAsString);
+		OutBeautifiedNames.Add(FPaths::GetBaseFilename(PathAsString));
+	}
+}
+
+
+bool FInterchangeImportTest::RunTest(const FString& Path)
+{
+	// Determine the test plan assets within the given path which will be run in parallel
+
+	TArray<FInterchangeImportTestData> TestPlans;
+
+	if (Path.Contains(TEXT(".")))
+	{
+		// Run test on a single TestPlan asset
+		FName PackageName = FName(FPaths::GetBaseFilename(Path, false));
+		FName PackagePath = FName(FPaths::GetPath(Path));
+		FName AssetName = FName(FPaths::GetBaseFilename(Path, true));
+		FName ClassName = UInterchangeImportTestPlan::StaticClass()->GetFName();
+		FInterchangeImportTestData& TP = TestPlans.AddDefaulted_GetRef();
+		TP.AssetData = FAssetData(PackageName, PackagePath, AssetName, ClassName);
+		TP.TestPlan = CastChecked<UInterchangeImportTestPlan>(TP.AssetData.GetAsset());
+	}
+	else
+	{
+		// Run tests in parallel on all TestPlan assets in the given directory
+		FARFilter AssetRegistryFilter;
+		AssetRegistryFilter.PackagePaths.Add(FName(Path));
+		AssetRegistryFilter.ClassNames.Add(UInterchangeImportTestPlan::StaticClass()->GetFName());
+		AssetRegistryFilter.bRecursivePaths = false;
+
+		const FAssetRegistryModule& AssetRegistryModule = FModuleManager::LoadModuleChecked<FAssetRegistryModule>(TEXT("AssetRegistry"));
+
+		AssetRegistryModule.Get().EnumerateAssets(AssetRegistryFilter,
+			[&TestPlans](const FAssetData& AssetData)
+			{
+				FInterchangeImportTestData& TP = TestPlans.AddDefaulted_GetRef();
+				TP.AssetData = AssetData;
+				TP.TestPlan = CastChecked<UInterchangeImportTestPlan>(AssetData.GetAsset());
+				return true;
+			}
+		);
+	}
+
+	// Base path to import assets to
+
+	FString SubDirToUse = TEXT("Interchange/ImportTest/");
+	FString BasePackagePath = FString(TEXT("/Temp/")) / SubDirToUse;
+	FString BaseFilePath = FPaths::ProjectSavedDir() / SubDirToUse;
+
+	// Clear out the folder contents before we do anything else
+
+	const bool bRequireExists = false;
+	const bool bDeleteRecursively = true;
+	IFileManager::Get().DeleteDirectory(*BaseFilePath, bRequireExists, bDeleteRecursively);
+
+	TArray<UE::Interchange::FAssetImportResultPtr> Results;
+
+	int32 StepIndex = 0;
+	bool bFoundValidSteps = true;
+	bool bSuccess = true;
+
+	while (bFoundValidSteps)
+	{
+		Results.Reset();
+		Results.SetNum(TestPlans.Num());
+
+		bFoundValidSteps = false;
+
+		// Start parallel import tasks in the interchange manager
+
+		for (int32 PlanIndex = 0; PlanIndex < TestPlans.Num(); PlanIndex++)
+		{
+			FInterchangeImportTestData& Data = TestPlans[PlanIndex];
+			
+			FString AssetName = Data.AssetData.AssetName.ToString().Replace(TEXT("/"), TEXT("_"));
+			Data.DestAssetPackagePath = BasePackagePath / AssetName;
+			Data.DestAssetFilePath = BaseFilePath / AssetName;
+
+			const bool bAddRecursively = true;
+			IFileManager::Get().MakeDirectory(*Data.DestAssetFilePath, bAddRecursively);
+
+			if (Data.TestPlan)
+			{
+				if (Data.TestPlan->Steps.IsValidIndex(StepIndex))
+				{
+					bFoundValidSteps = true;
+					Results[PlanIndex] = Data.TestPlan->Steps[StepIndex]->StartStep(Data);
+				}
+			}
+		}
+
+		if (bFoundValidSteps)
+		{
+			// If we started any test steps, join all the threads here, before proceeding with operations which have to be performed
+			// on the main thread, and which could potentially trigger GC.
+
+			for (int32 PlanIndex = 0; PlanIndex < TestPlans.Num(); PlanIndex++)
+			{
+				if (Results[PlanIndex])
+				{
+					Results[PlanIndex]->WaitUntilDone();
+				}
+			}
+
+			// Start post-steps
+
+			for (int32 PlanIndex = 0; PlanIndex < TestPlans.Num(); PlanIndex++)
+			{
+				FInterchangeImportTestData& Data = TestPlans[PlanIndex];
+
+				if (Data.TestPlan)
+				{
+					if (Data.TestPlan->Steps.IsValidIndex(StepIndex))
+					{
+						FString Context = Data.AssetData.AssetName.ToString() + TEXT(": ") + Data.TestPlan->Steps[StepIndex]->GetContextString();
+						ExecutionInfo.PushContext(Context);
+
+						// Fill out list of result objects in the data object.
+						// These are the UObject* results corresponding to imported assets.
+						Data.ResultObjects = Results[PlanIndex]->GetImportedObjects();
+
+						// Also add the InterchangeResultsContainer to the data so that tests can be run on it
+						// (e.g. to check whether something imported with a specific expected error)
+						Data.InterchangeResults = Results[PlanIndex]->GetResults();
+
+						// Populate the automation test execution info with the interchange import results
+						for (UInterchangeResult* Result : Data.InterchangeResults->GetResults())
+						{
+							switch (Result->GetResultType())
+							{
+							case EInterchangeResultType::Error:
+								ExecutionInfo.AddError(Result->GetText().ToString());
+								break;
+
+							case EInterchangeResultType::Warning:
+								ExecutionInfo.AddWarning(Result->GetText().ToString());
+								break;
+							}
+						}
+
+						// Fill out list of imported assets as FAssetData
+						for (UObject* Object : Results[PlanIndex]->GetImportedObjects())
+						{
+							Data.ImportedAssets.Emplace(Object);
+						}
+
+						// We don't need the Interchange results object any more, as we've already taken everything from it that we need.
+						// If we don't reset it, it will hold on to the trashed versions of the imported assets during GC.
+						Results[PlanIndex].Reset();
+
+						Data.TestPlan->Steps[StepIndex]->FinishStep(Data, ExecutionInfo);
+
+						ExecutionInfo.PopContext();
+					}
+				}
+			}
+
+			// Collect garbage between every step, so that 
+			CollectGarbage(GARBAGE_COLLECTION_KEEPFLAGS);
+		}
+
+		StepIndex++;
+	}
+
+	TArray<UObject*> ObjectsToDelete;
+	for (FInterchangeImportTestData& Data : TestPlans)
+	{
+		ObjectsToDelete.Append(Data.ResultObjects);
+	}
+
+	bool bShowConfirmation = false;
+	ObjectTools::ForceDeleteObjects(ObjectsToDelete, bShowConfirmation);
+
+	return true;
+}
diff --git a/Engine/Plugins/Tests/InterchangeTests/Source/InterchangeTests/Private/InterchangeImportTestPlan.cpp b/Engine/Plugins/Tests/InterchangeTests/Source/InterchangeTests/Private/InterchangeImportTestPlan.cpp
new file mode 100644
--- /dev/null
+++ b/Engine/Plugins/Tests/InterchangeTests/Source/InterchangeTests/Private/InterchangeImportTestPlan.cpp
@@ -0,0 +1,87 @@
+// Copyright Epic Games, Inc. All Rights Reserved.
+
+#include "InterchangeImportTestPlan.h"
+#include "JsonObjectConverter.h"
+#include "Misc/AutomationTest.h"
+#include "Logging/MessageLog.h"
+
+#define LOCTEXT_NAMESPACE "InterchangeImportTestPlan"
+
+
+bool UInterchangeImportTestPlan::ReadFromJson(const FString& Filename)
+{
+	FString JsonString;
+	if (!FFileHelper::LoadFileToString(JsonString, *Filename))
+	{
+		return false;
+	}
+
+	TSharedPtr<FJsonObject> JsonObject;
+	TSharedRef<TJsonReader<>> JsonReader = TJsonReaderFactory<>::Create(JsonString);
+	if (!FJsonSerializer::Deserialize(JsonReader, JsonObject) || !JsonObject.IsValid())
+	{
+		return false;
+	}
+
+	if (!FJsonObjectConverter::JsonObjectToUStruct(JsonObject.ToSharedRef(), UInterchangeImportTestPlan::StaticClass(), this))
+	{
+		return false;
+	}
+
+	return true;
+}
+
+
+void UInterchangeImportTestPlan::WriteToJson(const FString& Filename)
+{
+	FString JsonString;
+	if (!FJsonObjectConverter::UStructToJsonObjectString(UInterchangeImportTestPlan::StaticClass(), this, JsonString))
+	{
+		return;
+	}
+
+	FFileHelper::SaveStringToFile(JsonString, *Filename);
+}
+
+
+void UInterchangeImportTestPlan::RunThisTest()
+{
+	FMessageLog AutomationEditorLog("AutomationTestingLog");
+	FString NewPageName = FString::Printf(TEXT("----- Interchange Import Test: %s----"), *GetPathName());
+	FText NewPageNameText = FText::FromString(*NewPageName);
+	AutomationEditorLog.Open();
+	AutomationEditorLog.NewPage(NewPageNameText);
+	AutomationEditorLog.Info(NewPageNameText);
+
+	FAutomationTestFramework& TestFramework = FAutomationTestFramework::Get();
+	
+	TestFramework.StartTestByName(FString(TEXT("FInterchangeImportTest ")) + GetPathName(), 0);
+
+	FAutomationTestExecutionInfo ExecutionInfo;
+	if (TestFramework.StopTest(ExecutionInfo))
+	{
+		AutomationEditorLog.Info(LOCTEXT("TestPassed", "Passed"));
+	}
+	else
+	{
+		for (const auto& Entry : ExecutionInfo.GetEntries())
+		{
+			switch (Entry.Event.Type)
+			{
+			case EAutomationEventType::Error:
+				AutomationEditorLog.Error(FText::FromString(Entry.ToString()));
+				break;
+
+			case EAutomationEventType::Warning:
+				AutomationEditorLog.Warning(FText::FromString(Entry.ToString()));
+				break;
+
+			case EAutomationEventType::Info:
+				AutomationEditorLog.Info(FText::FromString(Entry.ToString()));
+				break;
+			}
+		}
+	}
+}
+
+#undef LOCTEXT_NAMESPACE
diff --git a/Engine/Plugins/Tests/InterchangeTests/Source/InterchangeTests/Private/InterchangeImportTestStepBase.cpp b/Engine/Plugins/Tests/InterchangeTests/Source/InterchangeTests/Private/InterchangeImportTestStepBase.cpp
new file mode 100644
--- /dev/null
+++ b/Engine/Plugins/Tests/InterchangeTests/Source/InterchangeTests/Private/InterchangeImportTestStepBase.cpp
@@ -0,0 +1,35 @@
+// Copyright Epic Games, Inc. All Rights Reserved.
+
+#include "InterchangeImportTestStepBase.h"
+#include "InterchangeImportTestData.h"
+#include "InterchangeTestFunction.h"
+
+
+bool UInterchangeImportTestStepBase::PerformTests(FInterchangeImportTestData& Data, FAutomationTestExecutionInfo& ExecutionInfo)
+{
+	// Add the Interchange results container to the list of UObject* we pass to the test functions.
+	// This way we can match against tests which can operate on the UInterchangeResultsContainer too.
+	TArray<UObject*> ResultObjects = Data.ResultObjects;
+	ResultObjects.Add(Data.InterchangeResults);
+
+	bool bSuccess = true;
+
+	for (FInterchangeTestFunction& Test : Tests)
+	{
+		FInterchangeTestFunctionResult Result = Test.Invoke(Data.ResultObjects);
+
+		for (const FString& Warning : Result.GetWarnings())
+		{
+			ExecutionInfo.AddWarning(Warning);
+		}
+
+		for (const FString& Error : Result.GetErrors())
+		{
+			ExecutionInfo.AddError(Error);
+		}
+
+		bSuccess &= Result.IsSuccess();
+	}
+
+	return bSuccess;
+}
diff --git a/Engine/Plugins/Tests/InterchangeTests/Source/InterchangeTests/Private/InterchangeImportTestStepImport.cpp b/Engine/Plugins/Tests/InterchangeTests/Source/InterchangeTests/Private/InterchangeImportTestStepImport.cpp
new file mode 100644
--- /dev/null
+++ b/Engine/Plugins/Tests/InterchangeTests/Source/InterchangeTests/Private/InterchangeImportTestStepImport.cpp
@@ -0,0 +1,119 @@
+// Copyright Epic Games, Inc. All Rights Reserved.
+
+#include "InterchangeImportTestStepImport.h"
+#include "InterchangeImportTestData.h"
+#include "InterchangeTestFunction.h"
+#include "UObject/SavePackage.h"
+#include "UObject/ReferenceChainSearch.h"
+#include "PackageTools.h"
+#include "Editor.h"
+
+
+UE::Interchange::FAssetImportResultPtr UInterchangeImportTestStepImport::StartStep(const FInterchangeImportTestData& Data)
+{
+	// Empty the destination folder here if requested
+	if (bEmptyDestinationFolderPriorToImport)
+	{
+		const bool bRequireExists = true;
+		const bool bDeleteRecursively = true;
+		IFileManager::Get().DeleteDirectory(*Data.DestAssetFilePath, bRequireExists, bDeleteRecursively);
+
+		const bool bAddRecursively = true;
+		IFileManager::Get().MakeDirectory(*Data.DestAssetFilePath, bAddRecursively);
+
+		// @todo: Is there a better way of deleting all the files inside a directory without also deleting the directory?
+	}
+
+	// Start the Interchange import
+	UE::Interchange::FScopedSourceData ScopedSourceData(SourceFile.FilePath);
+
+	FImportAssetParameters Params;
+	Params.OverridePipelines = PipelineStack;
+	Params.bIsAutomated = true;
+
+	UInterchangeManager& InterchangeManager = UInterchangeManager::GetInterchangeManager();
+	return InterchangeManager.ImportAssetAsync(Data.DestAssetPackagePath, ScopedSourceData.GetSourceData(), Params);
+}
+
+
+FTestStepResults UInterchangeImportTestStepImport::FinishStep(FInterchangeImportTestData& Data, FAutomationTestExecutionInfo& ExecutionInfo)
+{
+	FTestStepResults Results;
+
+	// If we need to perform a save and fresh reload of everything we imported, do it here
+	if (bSaveThenReloadImportedAssets)
+	{
+		// First save
+		for (const FAssetData& AssetData : Data.ImportedAssets)
+		{
+			const bool bLoadAsset = false;
+			UObject* AssetObject = AssetData.FastGetAsset(bLoadAsset);
+			UPackage* PackageObject = AssetObject->GetPackage();
+			check(PackageObject);
+
+			AssetObject->MarkPackageDirty();
+
+			FSavePackageArgs SaveArgs;
+			SaveArgs.TopLevelFlags = RF_Standalone;
+			SaveArgs.SaveFlags = SAVE_NoError;
+			UPackage::SavePackage(PackageObject, AssetObject,
+				*FPackageName::LongPackageNameToFilename(AssetData.PackageName.ToString(), FPackageName::GetAssetPackageExtension()),
+				SaveArgs);
+		}
+
+		if ((GEditor != nullptr) && (GEditor->Trans != nullptr))
+		{
+			GEditor->Trans->Reset(FText::FromString("Discard undo history during Automation testing."));
+		}
+
+		// Then rename original objects and their packages, and mark as garbage
+		for (const FAssetData& AssetData : Data.ImportedAssets)
+		{
+			const bool bLoadAsset = false;
+			UObject* AssetObject = AssetData.FastGetAsset(bLoadAsset);
+			UPackage* PackageObject = AssetObject->GetPackage();
+			check(PackageObject);
+			check(PackageObject == AssetData.GetPackage());
+
+			// Mark all objects in the package as garbage, and remove the standalone flag, so that GC can remove the package later
+			TArray<UObject*> ObjectsInPackage;
+			GetObjectsWithPackage(PackageObject, ObjectsInPackage, true);
+			for (UObject* ObjectInPackage : ObjectsInPackage)
+			{
+				ObjectInPackage->ClearFlags(RF_Standalone | RF_Public);
+				ObjectInPackage->MarkAsGarbage();
+			}
+
+			// Renaming the original objects avoids having to do a GC sweep here (this is done at the end of each test step)
+			// Any existing references to them will be retained but irrelevant.
+			// Then the new object can be loaded in their place, as if it were being loaded for the first time.
+			const ERenameFlags RenameFlags = REN_ForceNoResetLoaders | REN_DontCreateRedirectors | REN_NonTransactional | REN_DoNotDirty;
+			PackageObject->Rename(*(PackageObject->GetName() + TEXT("_TRASH")), nullptr, RenameFlags);
+			PackageObject->RemoveFromRoot();
+			PackageObject->MarkAsGarbage();
+		}
+
+
+		// Now reload
+		Data.ResultObjects.Reset();
+		for (const FAssetData& AssetData : Data.ImportedAssets)
+		{
+			check(!AssetData.IsAssetLoaded());
+			UObject* AssetObject = AssetData.GetAsset();
+			check(AssetObject);
+			Data.ResultObjects.Add(AssetObject);
+		}
+	}
+
+	// Run all the tests
+	bool bSuccess = PerformTests(Data, ExecutionInfo);
+
+	Results.bTestStepSuccess = bSuccess;
+	return Results;
+}
+
+
+FString UInterchangeImportTestStepImport::GetContextString() const
+{
+	return FString(TEXT("Importing ")) + FPaths::GetCleanFilename(SourceFile.FilePath);
+}
diff --git a/Engine/Plugins/Tests/InterchangeTests/Source/InterchangeTests/Private/InterchangeImportTestStepReimport.cpp b/Engine/Plugins/Tests/InterchangeTests/Source/InterchangeTests/Private/InterchangeImportTestStepReimport.cpp
new file mode 100644
--- /dev/null
+++ b/Engine/Plugins/Tests/InterchangeTests/Source/InterchangeTests/Private/InterchangeImportTestStepReimport.cpp
@@ -0,0 +1,126 @@
+// Copyright Epic Games, Inc. All Rights Reserved.
+
+#include "InterchangeImportTestStepReimport.h"
+#include "InterchangeImportTestData.h"
+#include "InterchangeTestFunction.h"
+#include "UObject/SavePackage.h"
+
+
+UE::Interchange::FAssetImportResultPtr UInterchangeImportTestStepReimport::StartStep(const FInterchangeImportTestData& Data)
+{
+	// Find the object we wish to reimport
+	TArray<UObject*> PotentialObjectsToReimport;
+	PotentialObjectsToReimport.Reserve(Data.ResultObjects.Num());
+
+	for (UObject* ResultObject : Data.ResultObjects)
+	{
+		if (ResultObject->GetClass() == AssetTypeToReimport.Get())
+		{
+			PotentialObjectsToReimport.Add(ResultObject);
+		}
+	}
+
+	UObject* AssetToReimport = nullptr;
+
+	if (PotentialObjectsToReimport.Num() == 1)
+	{
+		AssetToReimport = PotentialObjectsToReimport[0];
+	}
+	else if (PotentialObjectsToReimport.Num() > 1 && !AssetNameToReimport.IsEmpty())
+	{
+		for (UObject* Object : PotentialObjectsToReimport)
+		{
+			if (Object->GetName() == AssetNameToReimport)
+			{
+				AssetToReimport = Object;
+				break;
+			}
+		}
+	}
+
+	if (AssetToReimport == nullptr)
+	{
+		return UE::Interchange::FAssetImportResultPtr();
+	}
+
+	// Start the Interchange import
+	UE::Interchange::FScopedSourceData ScopedSourceData(SourceFileToReimport.FilePath);
+
+	FImportAssetParameters Params;
+	Params.bIsAutomated = true;
+	Params.ReimportAsset = AssetToReimport;
+
+	UInterchangeManager& InterchangeManager = UInterchangeManager::GetInterchangeManager();
+	return InterchangeManager.ImportAssetAsync(Data.DestAssetPackagePath, ScopedSourceData.GetSourceData(), Params);
+}
+
+
+FTestStepResults UInterchangeImportTestStepReimport::FinishStep(FInterchangeImportTestData& Data, FAutomationTestExecutionInfo& ExecutionInfo)
+{
+	FTestStepResults Results;
+
+	// If we need to perform a save and fresh reload of everything we imported, do it here
+	if (bSaveThenReloadImportedAssets)
+	{
+		// First save
+		for (const FAssetData& AssetData : Data.ImportedAssets)
+		{
+			const bool bLoadAsset = false;
+			UObject* AssetObject = AssetData.FastGetAsset(bLoadAsset);
+			UPackage* PackageObject = AssetObject->GetPackage();
+			check(PackageObject);
+
+			AssetObject->MarkPackageDirty();
+
+			FSavePackageArgs SaveArgs;
+			SaveArgs.TopLevelFlags = RF_Standalone;
+			SaveArgs.SaveFlags = SAVE_NoError;
+			UPackage::SavePackage(PackageObject, AssetObject,
+				*FPackageName::LongPackageNameToFilename(AssetData.PackageName.ToString(), FPackageName::GetAssetPackageExtension()),
+				SaveArgs);
+		}
+
+		// Then rename original objects and their packages, and mark as garbage
+		for (const FAssetData& AssetData : Data.ImportedAssets)
+		{
+			const bool bLoadAsset = false;
+			UObject* AssetObject = AssetData.FastGetAsset(bLoadAsset);
+			UPackage* PackageObject = AssetObject->GetPackage();
+			check(PackageObject);
+			check(PackageObject == AssetData.GetPackage());
+
+			// Renaming the original objects avoids having to do a GC sweep here.
+			// Any existing references to them will be retained but irrelevant.
+			// Then the new object can be loaded in their place, as if it were being loaded for the first time.
+			const ERenameFlags RenameFlags = REN_ForceNoResetLoaders | REN_DontCreateRedirectors | REN_NonTransactional | REN_DoNotDirty;
+			AssetObject->Rename(*(AssetObject->GetName() + TEXT("_TRASH")), PackageObject, RenameFlags);
+			PackageObject->Rename(*(PackageObject->GetName() + TEXT("_TRASH")), nullptr, RenameFlags);
+			AssetObject->ClearFlags(RF_Standalone | RF_Public);
+			AssetObject->RemoveFromRoot();
+			AssetObject->MarkAsGarbage();
+			PackageObject->MarkAsGarbage();
+		}
+
+		// Now reload
+		Data.ResultObjects.Reset();
+		for (const FAssetData& AssetData : Data.ImportedAssets)
+		{
+			check(!AssetData.IsAssetLoaded());
+			UObject* AssetObject = AssetData.GetAsset();
+			check(AssetObject);
+			Data.ResultObjects.Add(AssetObject);
+		}
+	}
+
+	// Run all the tests
+	bool bSuccess = PerformTests(Data, ExecutionInfo);
+
+	Results.bTestStepSuccess = bSuccess;
+	return Results;
+}
+
+
+FString UInterchangeImportTestStepReimport::GetContextString() const
+{
+	return FString(TEXT("Reimporting ")) + FPaths::GetCleanFilename(SourceFileToReimport.FilePath);
+}
diff --git a/Engine/Plugins/Tests/InterchangeTests/Source/InterchangeTests/Private/InterchangeTestFunction.cpp b/Engine/Plugins/Tests/InterchangeTests/Source/InterchangeTests/Private/InterchangeTestFunction.cpp
new file mode 100644
--- /dev/null
+++ b/Engine/Plugins/Tests/InterchangeTests/Source/InterchangeTests/Private/InterchangeTestFunction.cpp
@@ -0,0 +1,262 @@
+// Copyright Epic Games, Inc. All Rights Reserved.
+
+#include "InterchangeTestFunction.h"
+#include "ImportTestFunctions/ImportTestFunctionsBase.h"
+
+
+TSharedPtr<FStructOnScope> FInterchangeTestFunction::ImportParameters()
+{
+	// Clear the block of memory into which we will import the parameter values from text.
+	TSharedPtr<FStructOnScope> ParamData = MakeShared<FStructOnScope>(CheckFunction);
+
+	for (FProperty* ParamProperty : GetParameters())
+	{
+		ImportParameter(*ParamData, ParamProperty);
+	}
+
+	return ParamData;
+}
+
+
+void FInterchangeTestFunction::ImportParameter(FStructOnScope& ParamData, FProperty* ParamProperty)
+{
+	FName ParamName = ParamProperty->GetFName();
+	int32 Offset = ParamProperty->GetOffset_ForUFunction();
+
+	if (const FString* ParamValue = Parameters.Find(ParamName))
+	{
+		// Read param value from bound params
+		ParamProperty->ImportText(**ParamValue, ParamData.GetStructMemory() + Offset, 0, nullptr);
+	}
+	else
+	{
+		// Read param value from default metadata
+		FName MetadataKey(FString(TEXT("CPP_Default_")) + ParamName.ToString());
+		if (CheckFunction->HasMetaData(MetadataKey))
+		{
+			FString DefaultValue = CheckFunction->GetMetaData(MetadataKey);
+			ParamProperty->ImportText(*DefaultValue, ParamData.GetStructMemory() + Offset, 0, nullptr);
+		}
+	}
+}
+
+
+void FInterchangeTestFunction::ExportParameters()
+{
+#if 0
+	check(CheckFunction);
+	check(ParamData.IsValid());
+
+	Parameters.Empty(GetNumParameters());
+
+	for (FProperty* ParamProperty : GetParameters())
+	{
+		FName ParamName = ParamProperty->GetFName();
+		int32 Offset = ParamProperty->GetOffset_ForUFunction();
+
+		FString ExportedText;
+		ParamProperty->ExportTextItem(ExportedText, ParamData->GetStructMemory() + Offset, nullptr, nullptr, 0);
+		Parameters.Add(ParamName, ExportedText);
+	}
+#endif
+}
+
+
+TArray<UClass*> FInterchangeTestFunction::GetAvailableAssetClasses()
+{
+	TArray<UClass*> Classes;
+
+	// Add a null entry to represent "no class"
+	Classes.Add(nullptr);
+
+	TArray<UClass*> FunctionClasses;
+	GetDerivedClasses(UImportTestFunctionsBase::StaticClass(), FunctionClasses);
+
+	for (UClass* FunctionClass : FunctionClasses)
+	{
+		UClass* AssociatedClass = FunctionClass->GetDefaultObject<UImportTestFunctionsBase>()->GetAssociatedAssetType();
+		if (AssociatedClass != UObject::StaticClass())
+		{
+			Classes.AddUnique(AssociatedClass);
+		}
+	}
+
+	return Classes;
+}
+
+
+TArray<UFunction*> FInterchangeTestFunction::GetAvailableFunctions(const UClass* Class)
+{
+	TArray<UFunction*> Functions;
+
+	// Add a null entry to represent "no function"
+	Functions.Add(nullptr);
+
+	if (Class == nullptr)
+	{
+		return Functions;
+	}
+
+	TArray<UClass*> FunctionClasses;
+	GetDerivedClasses(UImportTestFunctionsBase::StaticClass(), FunctionClasses);
+
+	for (UClass* FunctionClass : FunctionClasses)
+	{
+		if (Class->IsChildOf(FunctionClass->GetDefaultObject<UImportTestFunctionsBase>()->GetAssociatedAssetType()))
+		{
+			ForEachObjectWithOuter(FunctionClass,
+				[&Functions](UObject* Object)
+				{
+					if (UFunction* Function = Cast<UFunction>(Object))
+					{
+						Functions.Add(Function);
+					}
+				},
+				false
+			);
+		}
+	}
+
+	return Functions;
+}
+
+
+bool FInterchangeTestFunction::IsValid(UFunction* Function) const
+{
+	// Null functions are an allowed special case 
+	if (Function == nullptr)
+	{
+		return true;
+	}
+
+	// The first parameter must be either UObject* or TArray<UObject*> that we imported.
+	FObjectProperty* FirstParam = CastField<FObjectProperty>(Function->ChildProperties);
+	if (FirstParam == nullptr)
+	{
+		FArrayProperty* FirstParamAsArray = CastField<FArrayProperty>(Function->ChildProperties);
+		if (FirstParamAsArray)
+		{
+			FirstParam = CastField<FObjectProperty>(FirstParamAsArray->Inner);
+		}
+	}
+
+	// Check that its type corresponds to the expected asset type
+	if (FirstParam == nullptr ||
+		FirstParam->ArrayDim != 1 ||
+		!AssetClass->IsChildOf(FirstParam->PropertyClass))
+	{
+		return false;
+	}
+
+	// Check that there are at least two parameters: the asset pointer, and the return value (which is counted as an out parameter)
+	if (Function->NumParms < 2)
+	{
+		return false;
+	}
+
+	// Check that the final parameter is really a return value, and not a function parameter
+	if (Function->ReturnValueOffset == MAX_uint16)
+	{
+		return false;
+	}
+
+	// Check that the return value is the correct type
+	FStructProperty* ReturnProperty = CastField<FStructProperty>(Function->GetReturnProperty());
+	if (ReturnProperty == nullptr ||
+		ReturnProperty->Struct != FInterchangeTestFunctionResult::StaticStruct())
+	{
+		return false;
+	}
+
+	return true;
+}
+
+
+bool FInterchangeTestFunction::RequiresArrayOfAssets() const
+{
+	if (CheckFunction == nullptr)
+	{
+		return false;
+	}
+
+	FArrayProperty* FirstParamAsArray = CastField<FArrayProperty>(CheckFunction->ChildProperties);
+	return (FirstParamAsArray != nullptr);
+}
+
+
+UClass* FInterchangeTestFunction::GetExpectedClass() const
+{
+	// Determine whether the first parameter is a simple UObject*
+	FObjectProperty* FirstParam = CastField<FObjectProperty>(CheckFunction->ChildProperties);
+	if (FirstParam == nullptr)
+	{
+		// If not, it must be an array, for this function to be valid
+		FArrayProperty* FirstParamAsArray = CastField<FArrayProperty>(CheckFunction->ChildProperties);
+		check(FirstParamAsArray);
+		FirstParam = CastField<FObjectProperty>(FirstParamAsArray->Inner);
+	}
+
+	return FirstParam->PropertyClass;
+}
+
+
+FInterchangeTestFunctionResult FInterchangeTestFunction::Invoke(const TArray<UObject*>& AssetsToTest)
+{
+	if (CheckFunction == nullptr)
+	{
+		return FInterchangeTestFunctionResult();
+	}
+
+	check(IsValid(CheckFunction));
+
+	// Build a list of the assets of the expected class type, to pass to the function
+	TArray<UObject*> AssetsToPass;
+	AssetsToPass.Reserve(AssetsToTest.Num());
+
+	UClass* ExpectedClass = GetExpectedClass();
+
+	for (UObject* Asset : AssetsToTest)
+	{
+		if (Asset->GetClass()->IsChildOf(ExpectedClass))
+		{
+			AssetsToPass.Add(Asset);
+		}
+	}
+
+	// Import the text parameters to binary
+	TSharedPtr<FStructOnScope> ParamData = ImportParameters();
+
+	if (RequiresArrayOfAssets())
+	{
+		// The first parameter is the array of UObject* that we imported.
+		// Copy its value straight into the parameter data from this function's argument.
+		FPlatformMemory::WriteUnaligned<TArray<UObject*>>(ParamData->GetStructMemory(), AssetsToPass);
+	}
+	else
+	{
+		// If there is not exactly one entry in the array, the result is ambiguous, so return with an error
+		if (AssetsToPass.Num() != 1)
+		{
+			FInterchangeTestFunctionResult Result;
+			Result.AddError(FString::Printf(
+				TEXT("Could not determine a single '%s' asset to run test on (there were %d possibilities)"),
+				*ExpectedClass->GetDisplayNameText().ToString(),
+				AssetsToPass.Num())
+			);
+
+			return Result;
+		}
+
+		// The first parameter is the single UObject* that we imported.
+		FPlatformMemory::WriteUnaligned<UObject*>(ParamData->GetStructMemory(), AssetsToPass[0]);
+	}
+
+	// Invoke the function now that the parameter binary data has been initialized
+	CheckFunction->GetOuter()->ProcessEvent(CheckFunction, ParamData->GetStructMemory());
+
+	// Copy out the return value.
+	// This creates a new instance of the return type, which is copy constructed here from the type-punned binary data.
+	// The binary data never exists as a "real" instance, so no destructor will ever be called there.
+	FInterchangeTestFunctionResult Result = FPlatformMemory::ReadUnaligned<FInterchangeTestFunctionResult>(ParamData->GetStructMemory() + CheckFunction->ReturnValueOffset);
+	return Result;
+}
diff --git a/Engine/Plugins/Tests/InterchangeTests/Source/InterchangeTests/Private/InterchangeTestsModule.cpp b/Engine/Plugins/Tests/InterchangeTests/Source/InterchangeTests/Private/InterchangeTestsModule.cpp
new file mode 100644
--- /dev/null
+++ b/Engine/Plugins/Tests/InterchangeTests/Source/InterchangeTests/Private/InterchangeTestsModule.cpp
@@ -0,0 +1,62 @@
+// Copyright Epic Games, Inc. All Rights Reserved.
+
+#include "InterchangeTestsModule.h"
+#include "InterchangeTestsLog.h"
+#include "Modules/ModuleInterface.h"
+#include "Modules/ModuleManager.h"
+#include "InterchangeImportTestSettings.h"
+
+
+#define LOCTEXT_NAMESPACE "InterchangeEditorModule"
+
+DEFINE_LOG_CATEGORY(LogInterchangeTests);
+
+
+FInterchangeTestsModule& FInterchangeTestsModule::Get()
+{
+	return FModuleManager::LoadModuleChecked<FInterchangeTestsModule>(INTERCHANGETESTS_MODULE_NAME);
+}
+
+
+bool FInterchangeTestsModule::IsAvailable()
+{
+	return FModuleManager::Get().IsModuleLoaded(INTERCHANGETESTS_MODULE_NAME);
+}
+
+
+void FInterchangeTestsModule::StartupModule()
+{
+}
+
+
+TArray<FString> FInterchangeTestsModule::GetImportTests() const
+{
+	const UInterchangeImportTestSettings* Settings = GetDefault<UInterchangeImportTestSettings>();
+	FString ContentDir = FPaths::ProjectContentDir();
+	FString Path = FPaths::Combine(ContentDir, Settings->ImportTestsPath);
+
+	TArray<FString> FilesInDirectory;
+	const bool bFindFiles = true;
+	const bool bFindDirectories = false;
+	IFileManager::Get().FindFilesRecursive(FilesInDirectory, *Path, TEXT("*.*"), bFindFiles, bFindDirectories);
+
+	TArray<FString> Results;
+	Results.Reserve(FilesInDirectory.Num());
+
+	for (FString& File : FilesInDirectory)
+	{
+		if (FPaths::GetExtension(File) == TEXT("json"))
+		{
+			Results.Emplace(MoveTemp(File));
+		}
+	}
+
+	return Results;
+}
+
+
+IMPLEMENT_MODULE(FInterchangeTestsModule, InterchangeTests);
+
+
+#undef LOCTEXT_NAMESPACE
+
diff --git a/Engine/Plugins/Tests/InterchangeTests/Source/InterchangeTests/Public/ImportTestFunctions/AnimationImportTestFunctions.h b/Engine/Plugins/Tests/InterchangeTests/Source/InterchangeTests/Public/ImportTestFunctions/AnimationImportTestFunctions.h
new file mode 100644
--- /dev/null
+++ b/Engine/Plugins/Tests/InterchangeTests/Source/InterchangeTests/Public/ImportTestFunctions/AnimationImportTestFunctions.h
@@ -0,0 +1,20 @@
+// Copyright Epic Games, Inc. All Rights Reserved.
+
+#pragma once
+
+#include "ImportTestFunctionsBase.h"
+#include "AnimationImportTestFunctions.generated.h"
+
+struct FInterchangeTestFunctionResult;
+
+
+UCLASS()
+class INTERCHANGETESTS_API UAnimationImportTestFunctions : public UImportTestFunctionsBase
+{
+	GENERATED_BODY()
+
+public:
+
+	// UImportTestFunctionsBase interface
+	virtual UClass* GetAssociatedAssetType() const override;
+};
diff --git a/Engine/Plugins/Tests/InterchangeTests/Source/InterchangeTests/Public/ImportTestFunctions/ImportTestFunctionsBase.h b/Engine/Plugins/Tests/InterchangeTests/Source/InterchangeTests/Public/ImportTestFunctions/ImportTestFunctionsBase.h
new file mode 100644
--- /dev/null
+++ b/Engine/Plugins/Tests/InterchangeTests/Source/InterchangeTests/Public/ImportTestFunctions/ImportTestFunctionsBase.h
@@ -0,0 +1,30 @@
+// Copyright Epic Games, Inc. All Rights Reserved.
+
+#pragma once
+
+#include "CoreMinimal.h"
+#include "UObject/ObjectMacros.h"
+#include "Engine/EngineTypes.h"
+#include "ImportTestFunctionsBase.generated.h"
+
+
+/**
+ * This is the base class for any static class which provides test functions for an asset type.
+ * Note that test functions defined in derived classes must be defined as UFUNCTION(Exec)
+ * This ensures that default parameters are held as metadata.
+ * However these functions are in reality not designed to be called from the console.
+ * @todo: create an alternative UFUNCTION tag for exporting default parameters as metadata.
+ */
+UCLASS()
+class INTERCHANGETESTS_API UImportTestFunctionsBase : public UObject
+{
+	GENERATED_BODY()
+
+public:
+
+	/** Determine which of the import test functions classes to use for a given asset class type */
+	static UClass* GetImportTestFunctionsClassForAssetType(UClass* AssetClass);
+
+	// UImportTestFunctionsBase interface
+	virtual UClass* GetAssociatedAssetType() const { return UObject::StaticClass(); }
+};
diff --git a/Engine/Plugins/Tests/InterchangeTests/Source/InterchangeTests/Public/ImportTestFunctions/InterchangeResultImportTestFunctions.h b/Engine/Plugins/Tests/InterchangeTests/Source/InterchangeTests/Public/ImportTestFunctions/InterchangeResultImportTestFunctions.h
new file mode 100644
--- /dev/null
+++ b/Engine/Plugins/Tests/InterchangeTests/Source/InterchangeTests/Public/ImportTestFunctions/InterchangeResultImportTestFunctions.h
@@ -0,0 +1,21 @@
+// Copyright Epic Games, Inc. All Rights Reserved.
+
+#pragma once
+
+#include "ImportTestFunctionsBase.h"
+#include "InterchangeResultImportTestFunctions.generated.h"
+
+struct FInterchangeTestFunctionResult;
+class UInterchangeResultsContainer;
+
+
+UCLASS()
+class INTERCHANGETESTS_API UInterchangeResultImportTestFunctions : public UImportTestFunctionsBase
+{
+	GENERATED_BODY()
+
+public:
+
+	// UImportTestFunctionsBase interface
+	virtual UClass* GetAssociatedAssetType() const override;
+};
diff --git a/Engine/Plugins/Tests/InterchangeTests/Source/InterchangeTests/Public/ImportTestFunctions/MaterialImportTestFunctions.h b/Engine/Plugins/Tests/InterchangeTests/Source/InterchangeTests/Public/ImportTestFunctions/MaterialImportTestFunctions.h
new file mode 100644
--- /dev/null
+++ b/Engine/Plugins/Tests/InterchangeTests/Source/InterchangeTests/Public/ImportTestFunctions/MaterialImportTestFunctions.h
@@ -0,0 +1,24 @@
+// Copyright Epic Games, Inc. All Rights Reserved.
+
+#pragma once
+
+#include "ImportTestFunctionsBase.h"
+#include "MaterialImportTestFunctions.generated.h"
+
+struct FInterchangeTestFunctionResult;
+
+
+UCLASS()
+class INTERCHANGETESTS_API UMaterialImportTestFunctions : public UImportTestFunctionsBase
+{
+	GENERATED_BODY()
+
+public:
+
+	// UImportTestFunctionsBase interface
+	virtual UClass* GetAssociatedAssetType() const override;
+
+	/** Check whether the expected number of materials are imported */
+	UFUNCTION(Exec)
+	static FInterchangeTestFunctionResult CheckImportedMaterialCount(const TArray<UMaterialInterface*>& Materials, int32 ExpectedNumberOfImportedMaterials);
+};
diff --git a/Engine/Plugins/Tests/InterchangeTests/Source/InterchangeTests/Public/ImportTestFunctions/SkeletalMeshImportTestFunctions.h b/Engine/Plugins/Tests/InterchangeTests/Source/InterchangeTests/Public/ImportTestFunctions/SkeletalMeshImportTestFunctions.h
new file mode 100644
--- /dev/null
+++ b/Engine/Plugins/Tests/InterchangeTests/Source/InterchangeTests/Public/ImportTestFunctions/SkeletalMeshImportTestFunctions.h
@@ -0,0 +1,69 @@
+// Copyright Epic Games, Inc. All Rights Reserved.
+
+#pragma once
+
+#include "ImportTestFunctionsBase.h"
+#include "SkeletalMeshImportTestFunctions.generated.h"
+
+struct FInterchangeTestFunctionResult;
+class USkeletalMesh;
+
+
+UCLASS()
+class INTERCHANGETESTS_API USkeletalMeshImportTestFunctions : public UImportTestFunctionsBase
+{
+	GENERATED_BODY()
+
+public:
+
+	// UImportTestFunctionsBase interface
+	virtual UClass* GetAssociatedAssetType() const override;
+
+	/** Check whether the expected number of skeletal meshes are imported */
+	UFUNCTION(Exec)
+	static FInterchangeTestFunctionResult CheckImportedSkeletalMeshCount(const TArray<USkeletalMesh*>& Meshes, int32 ExpectedNumberOfImportedSkeletalMeshes);
+
+	/** Check whether the vertex count in the built render data for the given LOD is as expected */
+	UFUNCTION(Exec)
+	static FInterchangeTestFunctionResult CheckRenderVertexCount(USkeletalMesh* Mesh, int32 LodIndex, int32 ExpectedNumberOfRenderVertices);
+
+	/** Check whether the mesh has the expected number of LODs */
+	UFUNCTION(Exec)
+	static FInterchangeTestFunctionResult CheckLodCount(USkeletalMesh* Mesh, int32 ExpectedNumberOfLods);
+
+	/** Check whether the mesh has the expected number of material slots */
+	UFUNCTION(Exec)
+	static FInterchangeTestFunctionResult CheckMaterialSlotCount(USkeletalMesh* Mesh, int32 ExpectedNumberOfMaterialSlots);
+
+	/** Check whether the built render data for the given mesh LOD has the expected number of sections */
+	UFUNCTION(Exec)
+	static FInterchangeTestFunctionResult CheckSectionCount(USkeletalMesh* Mesh, int32 LodIndex, int32 ExpectedNumberOfSections);
+
+	/** Check whether the given section in the built render data for the given LOD has the expected number of triangles */
+	UFUNCTION(Exec)
+	static FInterchangeTestFunctionResult CheckTriangleCountInSection(USkeletalMesh* Mesh, int32 LodIndex, int32 SectionIndex, int32 ExpectedNumberOfTriangles);
+
+	/** Check whether the mesh has the expected number of UV channels */
+	UFUNCTION(Exec)
+	static FInterchangeTestFunctionResult CheckUVChannelCount(USkeletalMesh* Mesh, int32 LodIndex, int32 ExpectedNumberOfUVChannels);
+
+	/** Check whether the material name for the given section in the render data for the given LOD is as expected */
+	UFUNCTION(Exec)
+	static FInterchangeTestFunctionResult CheckSectionMaterialName(USkeletalMesh* Mesh, int32 LodIndex, int32 SectionIndex, const FString& ExpectedMaterialName);
+
+	/** Check whether the vertex of the given index is at the expected position */
+	UFUNCTION(Exec)
+	static FInterchangeTestFunctionResult CheckVertexIndexPosition(USkeletalMesh* Mesh, int32 LodIndex, int32 VertexIndex, const FVector& ExpectedVertexPosition);
+
+	/** Check whether the mesh has the expected number of bones */
+	UFUNCTION(Exec)
+	static FInterchangeTestFunctionResult CheckBoneCount(USkeletalMesh* Mesh, int32 ExpectedNumberOfBones);
+
+	/** Check that the bone of the specified index has the expected position */
+	UFUNCTION(Exec)
+	static FInterchangeTestFunctionResult CheckBonePosition(USkeletalMesh* Mesh, int32 BoneIndex, const FVector& ExpectedBonePosition);
+
+	/** Check that the specified bone name is skinned with the expected number of vertices */
+	UFUNCTION(Exec)
+	static FInterchangeTestFunctionResult CheckSkinnedVertexCountForBone(USkeletalMesh* Mesh, const FString& BoneName, bool bTestFirstAlternateProfile, int32 ExpectedSkinnedVertexCount);
+};
diff --git a/Engine/Plugins/Tests/InterchangeTests/Source/InterchangeTests/Public/ImportTestFunctions/StaticMeshImportTestFunctions.h b/Engine/Plugins/Tests/InterchangeTests/Source/InterchangeTests/Public/ImportTestFunctions/StaticMeshImportTestFunctions.h
new file mode 100644
--- /dev/null
+++ b/Engine/Plugins/Tests/InterchangeTests/Source/InterchangeTests/Public/ImportTestFunctions/StaticMeshImportTestFunctions.h
@@ -0,0 +1,118 @@
+// Copyright Epic Games, Inc. All Rights Reserved.
+
+#pragma once
+
+#include "ImportTestFunctionsBase.h"
+#include "StaticMeshImportTestFunctions.generated.h"
+
+struct FInterchangeTestFunctionResult;
+class UStaticMesh;
+
+
+UENUM(meta = (Bitmask))
+enum class EStaticMeshImportTestGroundTruthBitflags : uint32
+{
+	CheckLodCountEqual                = (1 << 0),
+	CheckVertexCountEqual             = (1 << 1),
+	CheckTriangleCountEqual           = (1 << 2),
+	CheckUVChannelCountEqual          = (1 << 3),
+	CheckCollisionPrimitiveCountEqual = (1 << 4),
+	CheckVertexPositionsEqual         = (1 << 5),
+	CheckNormalsEqual                 = (1 << 6)
+};
+ENUM_CLASS_FLAGS(EStaticMeshImportTestGroundTruthBitflags);
+
+
+UCLASS()
+class INTERCHANGETESTS_API UStaticMeshImportTestFunctions : public UImportTestFunctionsBase
+{
+	GENERATED_BODY()
+
+public:
+
+	// UImportTestFunctionsBase interface
+	virtual UClass* GetAssociatedAssetType() const override;
+
+	/** Check whether the expected number of static meshes are imported */
+	UFUNCTION(Exec)
+	static FInterchangeTestFunctionResult CheckImportedStaticMeshCount(const TArray<UStaticMesh*>& Meshes, int32 ExpectedNumberOfImportedStaticMeshes);
+
+	/** Check whether the vertex count for the given LOD is as expected */
+	UFUNCTION(Exec)
+	static FInterchangeTestFunctionResult CheckVertexCount(UStaticMesh* Mesh, int32 LodIndex, int32 ExpectedNumberOfVertices);
+
+	/** Check whether the vertex count in the built render data for the given LOD is as expected */
+	UFUNCTION(Exec)
+	static FInterchangeTestFunctionResult CheckRenderVertexCount(UStaticMesh* Mesh, int32 LodIndex, int32 ExpectedNumberOfRenderVertices);
+
+	/** Check whether the mesh has the expected number of LODs */
+	UFUNCTION(Exec)
+	static FInterchangeTestFunctionResult CheckLodCount(UStaticMesh* Mesh, int32 ExpectedNumberOfLods);
+
+	/** Check whether the mesh has the expected number of material slots */
+	UFUNCTION(Exec)
+	static FInterchangeTestFunctionResult CheckMaterialSlotCount(UStaticMesh* Mesh, int32 ExpectedNumberOfMaterialSlots);
+
+	/** Check whether the mesh has the expected number of polygon groups for the given LOD */
+	UFUNCTION(Exec)
+	static FInterchangeTestFunctionResult CheckPolygonGroupCount(UStaticMesh* Mesh, int32 LodIndex, int32 ExpectedNumberOfPolygonGroups);
+
+	/** Check whether the built render data for the given mesh LOD has the expected number of sections */
+	UFUNCTION(Exec)
+	static FInterchangeTestFunctionResult CheckSectionCount(UStaticMesh* Mesh, int32 LodIndex, int32 ExpectedNumberOfSections);
+
+	/** Check whether the mesh has the expected number of triangles for the given LOD */
+	UFUNCTION(Exec)
+	static FInterchangeTestFunctionResult CheckTotalTriangleCount(UStaticMesh* Mesh, int32 LodIndex, int32 ExpectedTotalNumberOfTriangles);
+
+	/** Check whether the mesh has the expected number of polygons for the given LOD */
+	UFUNCTION(Exec)
+	static FInterchangeTestFunctionResult CheckTotalPolygonCount(UStaticMesh* Mesh, int32 LodIndex, int32 ExpectedNumberOfPolygons);
+
+	/** Check whether the mesh contains any quads or ngons */
+	UFUNCTION(Exec)
+	static FInterchangeTestFunctionResult CheckThatMeshHasQuadsOrNgons(UStaticMesh* Mesh, int32 LodIndex, bool bMeshHasQuadsOrNgons);
+
+	/** Check whether the given polygon group of the given LOD has the expected number of triangles */
+	UFUNCTION(Exec)
+	static FInterchangeTestFunctionResult CheckTriangleCountInPolygonGroup(UStaticMesh* Mesh, int32 LodIndex, int32 PolygonGroupIndex, int32 ExpectedNumberOfTriangles);
+
+	/** Check whether the given polygon group of the given LOD has the expected number of polygons */
+	UFUNCTION(Exec)
+	static FInterchangeTestFunctionResult CheckPolygonCountInPolygonGroup(UStaticMesh* Mesh, int32 LodIndex, int32 PolygonGroupIndex, int32 ExpectedNumberOfPolygons);
+
+	/** Check whether the mesh has the expected number of UV channels */
+	UFUNCTION(Exec)
+	static FInterchangeTestFunctionResult CheckUVChannelCount(UStaticMesh* Mesh, int32 LodIndex, int32 ExpectedNumberOfUVChannels);
+
+	/** Check whether the material slot name for the given polygon group in the given LOD is as expected */
+	UFUNCTION(Exec)
+	static FInterchangeTestFunctionResult CheckPolygonGroupImportedMaterialSlotName(UStaticMesh* Mesh, int32 LodIndex, int32 PolygonGroupIndex, const FString& ExpectedImportedMaterialSlotName);
+
+	/** Check whether the section index in the built render data for the given LOD is referencing the expected material index */
+	UFUNCTION(Exec)
+	static FInterchangeTestFunctionResult CheckSectionMaterialIndex(UStaticMesh* Mesh, int32 LodIndex, int32 SectionIndex, int32 ExpectedMaterialIndex);
+
+	/** Check whether the section index in the built render data for the given LOD is referencing the expected material */
+	UFUNCTION(Exec)
+	static FInterchangeTestFunctionResult CheckSectionMaterialName(UStaticMesh* Mesh, int32 LodIndex, int32 SectionIndex, const FString& ExpectedMaterialName);
+
+	/** Check whether the vertex of the given index is at the expected position */
+	UFUNCTION(Exec)
+	static FInterchangeTestFunctionResult CheckVertexIndexPosition(UStaticMesh* Mesh, int32 LodIndex, int32 VertexIndex, const FVector& ExpectedVertexPosition);
+
+	/** Check whether the expected number of simple collision primitives were imported */
+	UFUNCTION(Exec)
+	static FInterchangeTestFunctionResult CheckSimpleCollisionPrimitiveCount(UStaticMesh* Mesh, int32 ExpectedSphereElementCount, int32 ExpectedBoxElementCount, int32 ExpectedCapsuleElementCount, int32 ExpectedConvexElementCount, int32 ExpectedTaperedCapsuleElementCount);
+
+	/** Check whether the mesh is equivalent to a ground truth asset */
+	UFUNCTION(Exec)
+	static FInterchangeTestFunctionResult CheckAgainstGroundTruth(UStaticMesh* Mesh, TSoftObjectPtr<UStaticMesh> MeshToCompare,
+		bool bCheckVertexCountEqual = true,
+		bool bCheckTriangleCountEqual = true,
+		bool bCheckUVChannelCountEqual = true,
+		bool bCheckCollisionPrimitiveCountEqual = true,
+		bool bCheckVertexPositionsEqual = true,
+		bool bCheckNormalsEqual = true
+	);
+};
diff --git a/Engine/Plugins/Tests/InterchangeTests/Source/InterchangeTests/Public/ImportTestFunctions/TextureImportTestFunctions.h b/Engine/Plugins/Tests/InterchangeTests/Source/InterchangeTests/Public/ImportTestFunctions/TextureImportTestFunctions.h
new file mode 100644
--- /dev/null
+++ b/Engine/Plugins/Tests/InterchangeTests/Source/InterchangeTests/Public/ImportTestFunctions/TextureImportTestFunctions.h
@@ -0,0 +1,20 @@
+// Copyright Epic Games, Inc. All Rights Reserved.
+
+#pragma once
+
+#include "ImportTestFunctionsBase.h"
+#include "TextureImportTestFunctions.generated.h"
+
+struct FInterchangeTestFunctionResult;
+
+
+UCLASS()
+class INTERCHANGETESTS_API UTextureImportTestFunctions : public UImportTestFunctionsBase
+{
+	GENERATED_BODY()
+
+public:
+
+	// UImportTestFunctionsBase interface
+	virtual UClass* GetAssociatedAssetType() const override;
+};
diff --git a/Engine/Plugins/Tests/InterchangeTests/Source/InterchangeTests/Public/InterchangeImportTestData.h b/Engine/Plugins/Tests/InterchangeTests/Source/InterchangeTests/Public/InterchangeImportTestData.h
new file mode 100644
--- /dev/null
+++ b/Engine/Plugins/Tests/InterchangeTests/Source/InterchangeTests/Public/InterchangeImportTestData.h
@@ -0,0 +1,21 @@
+// Copyright Epic Games, Inc. All Rights Reserved.
+
+#pragma once
+
+#include "CoreMinimal.h"
+#include "AssetData.h"
+
+class UInterchangeImportTestPlan;
+class UInterchangeResultsContainer;
+
+
+struct FInterchangeImportTestData
+{
+	FString DestAssetPackagePath;
+	FString DestAssetFilePath;
+	FAssetData AssetData;
+	UInterchangeImportTestPlan* TestPlan = nullptr;
+	UInterchangeResultsContainer* InterchangeResults = nullptr;
+	TArray<UObject*> ResultObjects;
+	TArray<FAssetData> ImportedAssets;
+};
diff --git a/Engine/Plugins/Tests/InterchangeTests/Source/InterchangeTests/Public/InterchangeImportTestPlan.h b/Engine/Plugins/Tests/InterchangeTests/Source/InterchangeTests/Public/InterchangeImportTestPlan.h
new file mode 100644
--- /dev/null
+++ b/Engine/Plugins/Tests/InterchangeTests/Source/InterchangeTests/Public/InterchangeImportTestPlan.h
@@ -0,0 +1,38 @@
+// Copyright Epic Games, Inc. All Rights Reserved.
+
+#pragma once
+
+#include "CoreMinimal.h"
+#include "UObject/ObjectMacros.h"
+#include "Engine/EngineTypes.h"
+#include "InterchangeImportTestPlan.generated.h"
+
+class UInterchangeImportTestStepBase;
+
+
+/**
+* Define a test plan
+*/
+UCLASS(BlueprintType, EditInlineNew, PrioritizeCategories = Definition) // @TODO: can't find a way to force Run to be below Definition
+class INTERCHANGETESTS_API UInterchangeImportTestPlan : public UObject
+{
+	GENERATED_BODY()
+
+public:
+
+	/** Deserialize a test plan from Json */
+	bool ReadFromJson(const FString& Filename);
+
+	/* Serialize a test plan to Json */
+	void WriteToJson(const FString& Filename);
+
+public:
+
+	/** Set of steps to perform to carry out this test plan */
+	UPROPERTY(EditAnywhere, Instanced, Category = Definition, Meta = (DisplayPriority = 0))
+	TArray<UInterchangeImportTestStepBase*> Steps;
+
+	/** Click here to immediately run this single test through the automation framework */
+	UFUNCTION(BlueprintCallable, CallInEditor, Category = Run, Meta = (DisplayPriority = 1))
+	void RunThisTest();
+};
diff --git a/Engine/Plugins/Tests/InterchangeTests/Source/InterchangeTests/Public/InterchangeImportTestSettings.h b/Engine/Plugins/Tests/InterchangeTests/Source/InterchangeTests/Public/InterchangeImportTestSettings.h
new file mode 100644
--- /dev/null
+++ b/Engine/Plugins/Tests/InterchangeTests/Source/InterchangeTests/Public/InterchangeImportTestSettings.h
@@ -0,0 +1,26 @@
+// Copyright Epic Games, Inc. All Rights Reserved.
+
+#pragma once
+
+#include "CoreMinimal.h"
+#include "UObject/ObjectMacros.h"
+#include "UObject/Object.h"
+#include "Engine/EngineTypes.h"
+#include "InterchangeImportTestSettings.generated.h"
+
+
+/**
+ * Implement settings for the Interchange Import Test
+ */
+UCLASS(config=Engine, defaultconfig)
+class INTERCHANGETESTS_API UInterchangeImportTestSettings : public UObject
+{
+	GENERATED_BODY()
+
+public:
+	UPROPERTY(EditAnywhere, config, Category=Automation)
+	FString ImportTestsPath;
+
+	UPROPERTY(EditAnywhere, config, Category=Automation)
+	TArray<FString> ImportFiles;
+};
diff --git a/Engine/Plugins/Tests/InterchangeTests/Source/InterchangeTests/Public/InterchangeImportTestStepBase.h b/Engine/Plugins/Tests/InterchangeTests/Source/InterchangeTests/Public/InterchangeImportTestStepBase.h
new file mode 100644
--- /dev/null
+++ b/Engine/Plugins/Tests/InterchangeTests/Source/InterchangeTests/Public/InterchangeImportTestStepBase.h
@@ -0,0 +1,39 @@
+// Copyright Epic Games, Inc. All Rights Reserved.
+
+#pragma once
+
+#include "CoreMinimal.h"
+#include "UObject/ObjectMacros.h"
+#include "Engine/EngineTypes.h"
+#include "InterchangeImportTestData.h"
+#include "InterchangeManager.h"
+#include "InterchangeImportTestStepBase.generated.h"
+
+class FAutomationTestExecutionInfo;
+struct FInterchangeTestFunction;
+
+
+struct FTestStepResults
+{
+	bool bTestStepSuccess = true;
+	bool bTriggerGC = false;
+};
+
+
+UCLASS(BlueprintType, EditInlineNew, Abstract, autoExpandCategories = (General, Test))
+class INTERCHANGETESTS_API UInterchangeImportTestStepBase : public UObject
+{
+	GENERATED_BODY()
+
+public:
+	/** An array of results to check against */
+	UPROPERTY(EditAnywhere, Category = General)
+	TArray<FInterchangeTestFunction> Tests;
+
+public:
+	virtual UE::Interchange::FAssetImportResultPtr StartStep(const FInterchangeImportTestData& Data) PURE_VIRTUAL(UInterchangeImportTestStepBase::StartStep, return {}; );
+	virtual FTestStepResults FinishStep(FInterchangeImportTestData& Data, FAutomationTestExecutionInfo& ExecutionInfo) PURE_VIRTUAL(UInterchangeImportTestStepBase::FinishStep, return {}; );
+	virtual FString GetContextString() const PURE_VIRTUAL(UInterchangeImportTestStepBase::GetContextString, return {}; );
+
+	bool PerformTests(FInterchangeImportTestData& Data, FAutomationTestExecutionInfo& ExecutionInfo);
+};
diff --git a/Engine/Plugins/Tests/InterchangeTests/Source/InterchangeTests/Public/InterchangeImportTestStepImport.h b/Engine/Plugins/Tests/InterchangeTests/Source/InterchangeTests/Public/InterchangeImportTestStepImport.h
new file mode 100644
--- /dev/null
+++ b/Engine/Plugins/Tests/InterchangeTests/Source/InterchangeTests/Public/InterchangeImportTestStepImport.h
@@ -0,0 +1,40 @@
+// Copyright Epic Games, Inc. All Rights Reserved.
+
+#pragma once
+
+#include "InterchangeImportTestStepBase.h"
+#include "InterchangeImportTestStepImport.generated.h"
+
+
+class UInterchangePipelineBase;
+
+
+UCLASS(BlueprintType, Meta = (DisplayName = "Import a file"))
+class INTERCHANGETESTS_API UInterchangeImportTestStepImport : public UInterchangeImportTestStepBase
+{
+	GENERATED_BODY()
+
+public:
+
+	/** The source file to import (path relative to the json script) */
+	UPROPERTY(EditAnywhere, Category = General)
+	FFilePath SourceFile;
+
+	/** The pipeline stack to use when importing (an empty array will use the defaults) */
+	UPROPERTY(EditAnywhere, Instanced, Category = General)
+	TArray<TObjectPtr<UInterchangePipelineBase>> PipelineStack;
+
+	/** Whether the destination folder should be emptied prior to import */
+	UPROPERTY(EditAnywhere, Category = General)
+	bool bEmptyDestinationFolderPriorToImport = true;
+
+	/** Whether imported assets should be saved and freshly reloaded after import */
+	UPROPERTY(EditAnywhere, Category = General)
+	bool bSaveThenReloadImportedAssets = true;
+
+public:
+	// UInterchangeImportTestStepBase interface
+	virtual UE::Interchange::FAssetImportResultPtr StartStep(const FInterchangeImportTestData& Data) override;
+	virtual FTestStepResults FinishStep(FInterchangeImportTestData& Data, FAutomationTestExecutionInfo& ExecutionInfo) override;
+	virtual FString GetContextString() const override;
+};
diff --git a/Engine/Plugins/Tests/InterchangeTests/Source/InterchangeTests/Public/InterchangeImportTestStepReimport.h b/Engine/Plugins/Tests/InterchangeTests/Source/InterchangeTests/Public/InterchangeImportTestStepReimport.h
new file mode 100644
--- /dev/null
+++ b/Engine/Plugins/Tests/InterchangeTests/Source/InterchangeTests/Public/InterchangeImportTestStepReimport.h
@@ -0,0 +1,40 @@
+// Copyright Epic Games, Inc. All Rights Reserved.
+
+#pragma once
+
+#include "InterchangeImportTestStepBase.h"
+//#include "InterchangeTestFunction.h"
+#include "InterchangeImportTestStepReimport.generated.h"
+
+struct FInterchangeTestFunction;
+
+
+UCLASS(BlueprintType, Meta = (DisplayName = "Reimport a file"))
+class INTERCHANGETESTS_API UInterchangeImportTestStepReimport : public UInterchangeImportTestStepBase
+{
+	GENERATED_BODY()
+
+public:
+
+	/** The source file to import (path relative to the json script) */
+	UPROPERTY(EditAnywhere, Category = General)
+	FFilePath SourceFileToReimport;
+
+	/** The type of the asset to reimport. In the case that only one such asset were imported, this is unambiguous */
+	UPROPERTY(EditAnywhere, Category = General)
+	TSubclassOf<UObject> AssetTypeToReimport;
+
+	/** If there were multiple assets of the above type imported, specify the concrete name here */
+	UPROPERTY(EditAnywhere, Category = General)
+	FString AssetNameToReimport;
+
+	/** Whether imported assets should be saved and freshly reloaded after import */
+	UPROPERTY(EditAnywhere, Category = General)
+	bool bSaveThenReloadImportedAssets = true;
+
+public:
+	// UInterchangeImportTestStepBase interface
+	virtual UE::Interchange::FAssetImportResultPtr StartStep(const FInterchangeImportTestData& Data) override;
+	virtual FTestStepResults FinishStep(FInterchangeImportTestData& Data, FAutomationTestExecutionInfo& ExecutionInfo) override;
+	virtual FString GetContextString() const override;
+};
diff --git a/Engine/Plugins/Tests/InterchangeTests/Source/InterchangeTests/Public/InterchangeTestFunction.h b/Engine/Plugins/Tests/InterchangeTests/Source/InterchangeTests/Public/InterchangeTestFunction.h
new file mode 100644
--- /dev/null
+++ b/Engine/Plugins/Tests/InterchangeTests/Source/InterchangeTests/Public/InterchangeTestFunction.h
@@ -0,0 +1,166 @@
+// Copyright Epic Games, Inc. All Rights Reserved.
+
+#pragma once
+
+#include "CoreMinimal.h"
+#include "UObject/ObjectMacros.h"
+#include "Engine/EngineTypes.h"
+#include "InterchangeTestFunction.generated.h"
+
+
+USTRUCT()
+struct FInterchangeTestFunctionResult
+{
+	GENERATED_BODY()
+
+public:
+	void AddInfo(FString Message)
+	{
+		Infos.Add(MoveTemp(Message));
+	}
+
+	void AddWarning(FString Message)
+	{
+		Warnings.Add(MoveTemp(Message));
+	}
+
+	void AddError(FString Message)
+	{
+		Errors.Add(MoveTemp(Message));
+	}
+
+	bool IsSuccess() const { return Errors.Num() == 0; }
+
+	const TArray<FString>& GetInfos() const    { return Infos; }
+	const TArray<FString>& GetWarnings() const { return Warnings; }
+	const TArray<FString>& GetErrors() const   { return Errors; }
+
+private:
+	UPROPERTY()
+	TArray<FString> Infos;
+
+	UPROPERTY()
+	TArray<FString> Warnings;
+
+	UPROPERTY()
+	TArray<FString> Errors;
+};
+
+
+
+USTRUCT()
+struct INTERCHANGETESTS_API FInterchangeTestFunction
+{
+	GENERATED_BODY()
+
+public:
+	/** Type of the asset being tested */
+	UPROPERTY(EditAnywhere, Category = Test)
+	TSubclassOf<UObject> AssetClass;
+
+	/** A function to be called to determine whether the result is correct */
+	UPROPERTY(EditAnywhere, Category = Test)
+	TObjectPtr<UFunction> CheckFunction;
+
+	/** Parameters and their bound values as text */
+	UPROPERTY(EditAnywhere, Category = Test)
+	TMap<FName, FString> Parameters;
+
+public:
+	/** Import all parameter data from text to a new binary buffer */
+	TSharedPtr<FStructOnScope> ImportParameters();
+
+	/** Import a single parameter from text to an existing binary buffer */
+	void ImportParameter(FStructOnScope& ParamData, FProperty* Property);
+
+	/** Export parameter data from binary to text */
+	void ExportParameters();
+
+	/** Return a list of the supported asset classes */
+	static TArray<UClass*> GetAvailableAssetClasses();
+
+	/** Return a list of available functions which can act upon an asset of the given class */
+	static TArray<UFunction*> GetAvailableFunctions(const UClass* Class);
+
+	/** Return whether the given function is valid or not */
+	bool IsValid(UFunction* Function) const;
+
+	/** Return whether the function takes an array of assets */
+	bool RequiresArrayOfAssets() const;
+
+	/** Return the class type expected as the first parameter to the function */
+	UClass* GetExpectedClass() const;
+
+	/** Returns the number of parameters which must be bound to this function call */
+	int32 GetNumParameters() const { return CheckFunction->NumParms - 2; }
+
+	/** Call the test function bound with the held parameter values */
+	FInterchangeTestFunctionResult Invoke(const TArray<UObject*>& AssetsToTest);
+
+	/**
+	 * Provide a readable user-facing way of iterating the function parameters of interest.
+	 * Specifically, the first parameter is ignored (it should be a UObject* pointer to the asset which is being tested).
+	 * The final parameter is also ignored as it is the return value.
+	 */
+	class FParameters
+	{
+	public:
+
+		explicit FParameters(UFunction* InFunction)
+			: Function(InFunction)
+		{}
+
+		class TConstIterator
+		{
+		public:
+			explicit TConstIterator(const UFunction* InFunction, bool bIsEnd)
+			{
+				if (InFunction)
+				{
+					// We have specific requirements for the function:
+					// It must have at least one parameter (which must be a UObject* or derived), and it must have a return value.
+					// This is enforced in the UI, so we can fatal error here.
+					check(InFunction->NumParms > 2);
+					check(InFunction->ReturnValueOffset != MAX_uint16);
+					check(InFunction->ChildProperties);
+
+					Property = CastField<FProperty>(InFunction->ChildProperties->Next);
+					ParamNumber = bIsEnd ? InFunction->NumParms - 1 : 1;
+				}
+				else
+				{
+					Property = nullptr;
+					ParamNumber = 0;
+				}
+			}
+
+			TConstIterator& operator++()
+			{
+				Property = CastField<FProperty>(Property->Next);
+				++ParamNumber;
+				return *this;
+			}
+
+			FProperty* operator*() const
+			{
+				return Property;
+			}
+
+			friend bool operator==(const TConstIterator& A, const TConstIterator& B) { return A.ParamNumber == B.ParamNumber; }
+			friend bool operator!=(const TConstIterator& A, const TConstIterator& B) { return A.ParamNumber != B.ParamNumber; }
+
+		private:
+			FProperty* Property;
+			int32 ParamNumber;
+		};
+
+		TConstIterator begin() const { return TConstIterator(Function, false); }
+		TConstIterator end() const   { return TConstIterator(Function, true); }
+
+	private:
+		UFunction* Function;
+	};
+
+	/** Return a proxy object which is used to iterate through function parameters of interest */
+	FParameters GetParameters() const { return FParameters(CheckFunction); } 
+};
diff --git a/Engine/Plugins/Tests/InterchangeTests/Source/InterchangeTests/Public/InterchangeTestsLog.h b/Engine/Plugins/Tests/InterchangeTests/Source/InterchangeTests/Public/InterchangeTestsLog.h
new file mode 100644
--- /dev/null
+++ b/Engine/Plugins/Tests/InterchangeTests/Source/InterchangeTests/Public/InterchangeTestsLog.h
@@ -0,0 +1,6 @@
+// Copyright Epic Games, Inc. All Rights Reserved.
+#pragma once
+
+#include "Logging/LogMacros.h"
+
+DECLARE_LOG_CATEGORY_EXTERN(LogInterchangeTests, Log, All)
diff --git a/Engine/Plugins/Tests/InterchangeTests/Source/InterchangeTests/Public/InterchangeTestsModule.h b/Engine/Plugins/Tests/InterchangeTests/Source/InterchangeTests/Public/InterchangeTestsModule.h
new file mode 100644
--- /dev/null
+++ b/Engine/Plugins/Tests/InterchangeTests/Source/InterchangeTests/Public/InterchangeTestsModule.h
@@ -0,0 +1,23 @@
+// Copyright Epic Games, Inc. All Rights Reserved.
+#pragma once
+
+#include "Modules/ModuleInterface.h"
+#include "CoreMinimal.h"
+
+#define INTERCHANGETESTS_MODULE_NAME TEXT("InterchangeTests")
+
+/**
+ * Module for implementing Interchange automation tests
+ */
+class INTERCHANGETESTS_API FInterchangeTestsModule : public IModuleInterface
+{
+public:
+	static FInterchangeTestsModule& Get();
+	static bool IsAvailable();
+
+	/** Get a list of json files defining import tests */
+	TArray<FString> GetImportTests() const;
+
+private:
+	virtual void StartupModule() override;
+};
diff --git a/Engine/Source/Runtime/Experimental/Interchange/Core/Public/InterchangePipelineBase.h b/Engine/Source/Runtime/Experimental/Interchange/Core/Public/InterchangePipelineBase.h
--- a/Engine/Source/Runtime/Experimental/Interchange/Core/Public/InterchangePipelineBase.h
+++ b/Engine/Source/Runtime/Experimental/Interchange/Core/Public/InterchangePipelineBase.h
@@ -19,7 +19,7 @@ enum class EInterchangePipelineTask : uint8
 	Export
 };
 
-UCLASS(BlueprintType, Blueprintable, Experimental)
+UCLASS(BlueprintType, Blueprintable, Experimental, EditInlineNew, Abstract)
 class INTERCHANGECORE_API UInterchangePipelineBase : public UObject
 {
 	GENERATED_BODY()
diff --git a/Engine/Source/Runtime/Experimental/Interchange/Engine/Public/InterchangeManager.h b/Engine/Source/Runtime/Experimental/Interchange/Engine/Public/InterchangeManager.h
--- a/Engine/Source/Runtime/Experimental/Interchange/Engine/Public/InterchangeManager.h
+++ b/Engine/Source/Runtime/Experimental/Interchange/Engine/Public/InterchangeManager.h
@@ -154,6 +154,8 @@ namespace UE
 
 		using FAssetImportResultRef = TSharedRef< FImportResult, ESPMode::ThreadSafe >;
 		using FSceneImportResultRef = TSharedRef< FImportResult, ESPMode::ThreadSafe >;
+		using FAssetImportResultPtr = TSharedPtr< FImportResult, ESPMode::ThreadSafe >;
+		using FSceneImportResultPtr = TSharedPtr< FImportResult, ESPMode::ThreadSafe >;
 
 		class FImportAsyncHelper : protected FGCObject
 		{
