Fix uninitialized member in FInterchangeTestFunction.

diff --git a/Engine/Plugins/Tests/InterchangeTests/Source/InterchangeTests/Public/InterchangeTestFunction.h b/Engine/Plugins/Tests/InterchangeTests/Source/InterchangeTests/Public/InterchangeTestFunction.h
--- a/Engine/Plugins/Tests/InterchangeTests/Source/InterchangeTests/Public/InterchangeTestFunction.h
+++ b/Engine/Plugins/Tests/InterchangeTests/Source/InterchangeTests/Public/InterchangeTestFunction.h
@@ -60,7 +60,7 @@ public:
 
 	/** A function to be called to determine whether the result is correct */
 	UPROPERTY(EditAnywhere, Category = Test)
-	TObjectPtr<UFunction> CheckFunction;
+	TObjectPtr<UFunction> CheckFunction = nullptr;
 
 	/** Parameters and their bound values as text */
 	UPROPERTY(EditAnywhere, Category = Test)
