Replaced raw UObject* property with TObjectPtr. Added missing include. (Why wasn't this stuff picked up by preflight!)

diff --git a/Engine/Plugins/Tests/InterchangeTests/Source/InterchangeTests/Private/InterchangeImportTestPlan.cpp b/Engine/Plugins/Tests/InterchangeTests/Source/InterchangeTests/Private/InterchangeImportTestPlan.cpp
--- a/Engine/Plugins/Tests/InterchangeTests/Source/InterchangeTests/Private/InterchangeImportTestPlan.cpp
+++ b/Engine/Plugins/Tests/InterchangeTests/Source/InterchangeTests/Private/InterchangeImportTestPlan.cpp
@@ -3,6 +3,7 @@
 #include "InterchangeImportTestPlan.h"
 #include "JsonObjectConverter.h"
 #include "Misc/AutomationTest.h"
+#include "Misc/FileHelper.h"
 #include "Logging/MessageLog.h"
 
 #define LOCTEXT_NAMESPACE "InterchangeImportTestPlan"
diff --git a/Engine/Plugins/Tests/InterchangeTests/Source/InterchangeTests/Public/InterchangeImportTestPlan.h b/Engine/Plugins/Tests/InterchangeTests/Source/InterchangeTests/Public/InterchangeImportTestPlan.h
--- a/Engine/Plugins/Tests/InterchangeTests/Source/InterchangeTests/Public/InterchangeImportTestPlan.h
+++ b/Engine/Plugins/Tests/InterchangeTests/Source/InterchangeTests/Public/InterchangeImportTestPlan.h
@@ -30,7 +30,7 @@ public:
 
 	/** Set of steps to perform to carry out this test plan */
 	UPROPERTY(EditAnywhere, Instanced, Category = Definition, Meta = (DisplayPriority = 0))
-	TArray<UInterchangeImportTestStepBase*> Steps;
+	TArray<TObjectPtr<UInterchangeImportTestStepBase>> Steps;
 
 	/** Click here to immediately run this single test through the automation framework */
 	UFUNCTION(BlueprintCallable, CallInEditor, Category = Run, Meta = (DisplayPriority = 1))
