Fixup deprecation warnings
#rnx

diff --git a/Engine/Plugins/Animation/ControlRig/Source/ControlRigEditor/Private/Editor/DetailsViewWrapperObject.cpp b/Engine/Plugins/Animation/ControlRig/Source/ControlRigEditor/Private/Editor/DetailsViewWrapperObject.cpp
--- a/Engine/Plugins/Animation/ControlRig/Source/ControlRigEditor/Private/Editor/DetailsViewWrapperObject.cpp
+++ b/Engine/Plugins/Animation/ControlRig/Source/ControlRigEditor/Private/Editor/DetailsViewWrapperObject.cpp
@@ -596,6 +596,6 @@ void UDetailsViewWrapperObject::SetContentForPin(URigVMPin* InPin)
 	{
 		FDetailsViewWrapperObjectImportErrorContext ErrorPipe;
 		uint8* DataPtr = Property->ContainerPtrToValuePtr<uint8>(this);
-		Property->ImportText(*DefaultValue, DataPtr, EPropertyPortFlags::PPF_None, nullptr, &ErrorPipe);		
+		Property->ImportText_InContainer(*DefaultValue, DataPtr, nullptr, EPropertyPortFlags::PPF_None, &ErrorPipe);		
 	}
 }
diff --git a/Engine/Plugins/Runtime/nDisplay/Source/DisplayClusterConfigurator/Private/DisplayClusterConfiguratorCompiler.cpp b/Engine/Plugins/Runtime/nDisplay/Source/DisplayClusterConfigurator/Private/DisplayClusterConfiguratorCompiler.cpp
--- a/Engine/Plugins/Runtime/nDisplay/Source/DisplayClusterConfigurator/Private/DisplayClusterConfiguratorCompiler.cpp
+++ b/Engine/Plugins/Runtime/nDisplay/Source/DisplayClusterConfigurator/Private/DisplayClusterConfiguratorCompiler.cpp
@@ -219,8 +219,8 @@ void FDisplayClusterConfiguratorKismetCompilerContext::CopyTermDefaultsToDefault
 					void* PropertyValue = Property->ContainerPtrToValuePtr<void*>(OldSubObject);
 			
 					FString Value;
-					Property->ExportTextItem(Value, PropertyValue, nullptr, nullptr, 0, nullptr);
-					Property->ImportText(*Value, PropertyValue, 0, OldSubObject);
+					Property->ExportTextItem_Direct(Value, PropertyValue, nullptr, nullptr, 0, nullptr);
+					Property->ImportText_Direct(*Value, PropertyValue, OldSubObject, 0);
 				}
 			}
 		}
diff --git a/Engine/Plugins/Tests/InterchangeTests/Source/InterchangeTests/Private/InterchangeTestFunction.cpp b/Engine/Plugins/Tests/InterchangeTests/Source/InterchangeTests/Private/InterchangeTestFunction.cpp
--- a/Engine/Plugins/Tests/InterchangeTests/Source/InterchangeTests/Private/InterchangeTestFunction.cpp
+++ b/Engine/Plugins/Tests/InterchangeTests/Source/InterchangeTests/Private/InterchangeTestFunction.cpp
@@ -26,7 +26,7 @@ void FInterchangeTestFunction::ImportParameter(FStructOnScope& ParamData, FPrope
 	if (const FString* ParamValue = Parameters.Find(ParamName))
 	{
 		// Read param value from bound params
-		ParamProperty->ImportText(**ParamValue, ParamData.GetStructMemory() + Offset, 0, nullptr);
+		ParamProperty->ImportText_InContainer(**ParamValue, ParamData.GetStructMemory() + Offset, nullptr, 0);
 	}
 	else
 	{
@@ -35,7 +35,7 @@ void FInterchangeTestFunction::ImportParameter(FStructOnScope& ParamData, FPrope
 		if (CheckFunction->HasMetaData(MetadataKey))
 		{
 			FString DefaultValue = CheckFunction->GetMetaData(MetadataKey);
-			ParamProperty->ImportText(*DefaultValue, ParamData.GetStructMemory() + Offset, 0, nullptr);
+			ParamProperty->ImportText_InContainer(*DefaultValue, ParamData.GetStructMemory() + Offset, nullptr, 0);
 		}
 	}
 }
@@ -55,7 +55,7 @@ void FInterchangeTestFunction::ExportParameters()
 		int32 Offset = ParamProperty->GetOffset_ForUFunction();
 
 		FString ExportedText;
-		ParamProperty->ExportTextItem(ExportedText, ParamData->GetStructMemory() + Offset, nullptr, nullptr, 0);
+		ParamProperty->ExportTextItem_InContainer(ExportedText, ParamData->GetStructMemory() + Offset, nullptr, nullptr, 0);
 		Parameters.Add(ParamName, ExportedText);
 	}
 #endif
