Fixed non-unity build error.
Fixed deprecation warnings relating to property import.

diff --git a/Engine/Plugins/Tests/InterchangeTests/Source/InterchangeTests/Private/InterchangeTestFunction.cpp b/Engine/Plugins/Tests/InterchangeTests/Source/InterchangeTests/Private/InterchangeTestFunction.cpp
--- a/Engine/Plugins/Tests/InterchangeTests/Source/InterchangeTests/Private/InterchangeTestFunction.cpp
+++ b/Engine/Plugins/Tests/InterchangeTests/Source/InterchangeTests/Private/InterchangeTestFunction.cpp
@@ -2,6 +2,7 @@
 
 #include "InterchangeTestFunction.h"
 #include "ImportTestFunctions/ImportTestFunctionsBase.h"
+#include "UObject/StructOnScope.h"
 
 
 TSharedPtr<FStructOnScope> FInterchangeTestFunction::ImportParameters()
@@ -26,7 +27,7 @@ void FInterchangeTestFunction::ImportParameter(FStructOnScope& ParamData, FPrope
 	if (const FString* ParamValue = Parameters.Find(ParamName))
 	{
 		// Read param value from bound params
-		ParamProperty->ImportText_Direct(**ParamValue, ParamData.GetStructMemory() + Offset, nullptr, 0);
+		ParamProperty->ImportText_Direct(**ParamValue, ParamData.GetStructMemory() + Offset, nullptr, PPF_None);
 	}
 	else
 	{
@@ -35,7 +36,7 @@ void FInterchangeTestFunction::ImportParameter(FStructOnScope& ParamData, FPrope
 		if (CheckFunction->HasMetaData(MetadataKey))
 		{
 			FString DefaultValue = CheckFunction->GetMetaData(MetadataKey);
-			ParamProperty->ImportText_Direct(*DefaultValue, ParamData.GetStructMemory() + Offset, nullptr, 0);
+			ParamProperty->ImportText_Direct(*DefaultValue, ParamData.GetStructMemory() + Offset, nullptr, PPF_None);
 		}
 	}
 }
diff --git a/Engine/Plugins/Tests/InterchangeTests/Source/InterchangeTests/Public/InterchangeTestFunction.h b/Engine/Plugins/Tests/InterchangeTests/Source/InterchangeTests/Public/InterchangeTestFunction.h
--- a/Engine/Plugins/Tests/InterchangeTests/Source/InterchangeTests/Public/InterchangeTestFunction.h
+++ b/Engine/Plugins/Tests/InterchangeTests/Source/InterchangeTests/Public/InterchangeTestFunction.h
@@ -7,6 +7,8 @@
 #include "Engine/EngineTypes.h"
 #include "InterchangeTestFunction.generated.h"
 
+class FStructOnScope;
+
 
 USTRUCT()
 struct FInterchangeTestFunctionResult
