Build fix for InterchangeTests: added missing includes

diff --git a/Engine/Plugins/Tests/InterchangeTests/Source/InterchangeTestEditor/Private/InterchangeTestFunctionLayout.cpp b/Engine/Plugins/Tests/InterchangeTests/Source/InterchangeTestEditor/Private/InterchangeTestFunctionLayout.cpp
--- a/Engine/Plugins/Tests/InterchangeTests/Source/InterchangeTestEditor/Private/InterchangeTestFunctionLayout.cpp
+++ b/Engine/Plugins/Tests/InterchangeTests/Source/InterchangeTestEditor/Private/InterchangeTestFunctionLayout.cpp
@@ -1,13 +1,16 @@
 // Copyright Epic Games, Inc. All Rights Reserved.
 
 #include "InterchangeTestFunctionLayout.h"
+#include "InterchangeTestFunction.h"
+
 #include "IDetailChildrenBuilder.h"
 #include "IPropertyUtilities.h"
 #include "DetailWidgetRow.h"
 #include "DetailLayoutBuilder.h"
+#include "Editor.h"
 #include "Widgets/Input/SNumericEntryBox.h"
 #include "Widgets/Input/SComboBox.h"
-#include "InterchangeTestFunction.h"
+#include "ScopedTransaction.h"
 
 #define LOCTEXT_NAMESPACE "InterchangeTestFunctionLayout"
 
diff --git a/Engine/Plugins/Tests/InterchangeTests/Source/InterchangeTests/Private/ImportTestFunctions/SkeletalMeshImportTestFunctions.cpp b/Engine/Plugins/Tests/InterchangeTests/Source/InterchangeTests/Private/ImportTestFunctions/SkeletalMeshImportTestFunctions.cpp
--- a/Engine/Plugins/Tests/InterchangeTests/Source/InterchangeTests/Private/ImportTestFunctions/SkeletalMeshImportTestFunctions.cpp
+++ b/Engine/Plugins/Tests/InterchangeTests/Source/InterchangeTests/Private/ImportTestFunctions/SkeletalMeshImportTestFunctions.cpp
@@ -2,6 +2,8 @@
 
 #include "ImportTestFunctions/SkeletalMeshImportTestFunctions.h"
 #include "InterchangeTestFunction.h"
+
+#include "Animation/Skeleton.h"
 #include "Engine/SkeletalMesh.h"
 #include "Rendering/SkeletalMeshModel.h"
 #include "Rendering/SkeletalMeshRenderData.h"
