Fix for static analysis in InterchangeTests: Switched to lvalues for mesh description attributes since we're keeping references on the elements.

diff --git a/Engine/Plugins/Tests/InterchangeTests/Source/InterchangeTests/Private/ImportTestFunctions/StaticMeshImportTestFunctions.cpp b/Engine/Plugins/Tests/InterchangeTests/Source/InterchangeTests/Private/ImportTestFunctions/StaticMeshImportTestFunctions.cpp
--- a/Engine/Plugins/Tests/InterchangeTests/Source/InterchangeTests/Private/ImportTestFunctions/StaticMeshImportTestFunctions.cpp
+++ b/Engine/Plugins/Tests/InterchangeTests/Source/InterchangeTests/Private/ImportTestFunctions/StaticMeshImportTestFunctions.cpp
@@ -4,6 +4,7 @@
 #include "InterchangeTestFunction.h"
 #include "Engine/StaticMesh.h"
 #include "Engine/SkeletalMesh.h"
+#include "PhysicsEngine/BodySetup.h"
 #include "StaticMeshAttributes.h"
 
 
@@ -389,7 +390,7 @@ FInterchangeTestFunctionResult UStaticMeshImportTestFunctions::CheckVertexIndexP
 		}
 		else
 		{
-			const FVector& VertexPosition = Attributes.GetVertexPositions()[VertexIndex];
+			const FVector3f VertexPosition = Attributes.GetVertexPositions()[VertexIndex];
 			if (!VertexPosition.Equals(ExpectedVertexPosition))
 			{
 				Result.AddError(FString::Printf(TEXT("For LOD %d vertex index %d, expected position %s, imported %s."), LodIndex, VertexIndex, *ExpectedVertexPosition.ToString(), *VertexPosition.ToString()));
@@ -579,7 +580,10 @@ FInterchangeTestFunctionResult UStaticMeshImportTestFunctions::CheckAgainstGroun
 			const FMeshDescription* GroundTruthMD = GroundTruth->GetMeshDescription(LodIndex);
 
 			FStaticMeshConstAttributes Attributes(*MD);
+			TVertexAttributesConstRef<FVector3f> VertexPositions = Attributes.GetVertexPositions();
+
 			FStaticMeshConstAttributes GroundTruthAttributes(*GroundTruthMD);
+			TVertexAttributesConstRef<FVector3f> GroundTruthVertexPositions = GroundTruthAttributes.GetVertexPositions();
 
 			int32 VertexCount = Attributes.GetVertexPositions().GetNumElements();
 			int32 ExpectedVertexCount = GroundTruthAttributes.GetVertexPositions().GetNumElements();
@@ -592,8 +596,8 @@ FInterchangeTestFunctionResult UStaticMeshImportTestFunctions::CheckAgainstGroun
 			{
 				for (int32 VertexIndex = 0; VertexIndex < VertexCount; VertexIndex++)
 				{
-					const FVector& V0 = Attributes.GetVertexPositions()[VertexIndex];
-					const FVector& V1 = GroundTruthAttributes.GetVertexPositions()[VertexIndex];
+					const FVector& V0 = VertexPositions[VertexIndex];
+					const FVector& V1 = GroundTruthVertexPositions[VertexIndex];
 
 					if (!V0.Equals(V1))
 					{
@@ -612,7 +616,10 @@ FInterchangeTestFunctionResult UStaticMeshImportTestFunctions::CheckAgainstGroun
 			const FMeshDescription* GroundTruthMD = GroundTruth->GetMeshDescription(LodIndex);
 
 			FStaticMeshConstAttributes Attributes(*MD);
+			TVertexInstanceAttributesConstRef<FVector3f> VertexInstanceNormals = Attributes.GetVertexInstanceNormals();
+
 			FStaticMeshConstAttributes GroundTruthAttributes(*GroundTruthMD);
+			TVertexInstanceAttributesConstRef<FVector3f> GroundTruthVertexInstanceNormals = GroundTruthAttributes.GetVertexInstanceNormals();
 
 			int32 NormalsCount = Attributes.GetVertexInstanceNormals().GetNumElements();
 			int32 ExpectedNormalsCount = GroundTruthAttributes.GetVertexInstanceNormals().GetNumElements();
@@ -625,8 +632,8 @@ FInterchangeTestFunctionResult UStaticMeshImportTestFunctions::CheckAgainstGroun
 			{
 				for (int32 NormalIndex = 0; NormalIndex < NormalsCount; NormalIndex++)
 				{
-					const FVector& N0 = Attributes.GetVertexInstanceNormals()[NormalIndex];
-					const FVector& N1 = GroundTruthAttributes.GetVertexInstanceNormals()[NormalIndex];
+					const FVector& N0 = VertexInstanceNormals[NormalIndex];
+					const FVector& N1 = GroundTruthVertexInstanceNormals[NormalIndex];
 
 					if (!N0.Equals(N1))
 					{
