Removed function which is no longer used.

diff --git a/Engine/Plugins/Tests/InterchangeTests/Source/InterchangeTests/Private/InterchangeTestsModule.cpp b/Engine/Plugins/Tests/InterchangeTests/Source/InterchangeTests/Private/InterchangeTestsModule.cpp
--- a/Engine/Plugins/Tests/InterchangeTests/Source/InterchangeTests/Private/InterchangeTestsModule.cpp
+++ b/Engine/Plugins/Tests/InterchangeTests/Source/InterchangeTests/Private/InterchangeTestsModule.cpp
@@ -29,32 +29,6 @@ void FInterchangeTestsModule::StartupModule()
 }
 
 
-TArray<FString> FInterchangeTestsModule::GetImportTests() const
-{
-	const UInterchangeImportTestSettings* Settings = GetDefault<UInterchangeImportTestSettings>();
-	FString ContentDir = FPaths::ProjectContentDir();
-	FString Path = FPaths::Combine(ContentDir, Settings->ImportTestsPath);
-
-	TArray<FString> FilesInDirectory;
-	const bool bFindFiles = true;
-	const bool bFindDirectories = false;
-	IFileManager::Get().FindFilesRecursive(FilesInDirectory, *Path, TEXT("*.*"), bFindFiles, bFindDirectories);
-
-	TArray<FString> Results;
-	Results.Reserve(FilesInDirectory.Num());
-
-	for (FString& File : FilesInDirectory)
-	{
-		if (FPaths::GetExtension(File) == TEXT("json"))
-		{
-			Results.Emplace(MoveTemp(File));
-		}
-	}
-
-	return Results;
-}
-
-
 IMPLEMENT_MODULE(FInterchangeTestsModule, InterchangeTests);
 
 
diff --git a/Engine/Plugins/Tests/InterchangeTests/Source/InterchangeTests/Public/InterchangeTestsModule.h b/Engine/Plugins/Tests/InterchangeTests/Source/InterchangeTests/Public/InterchangeTestsModule.h
--- a/Engine/Plugins/Tests/InterchangeTests/Source/InterchangeTests/Public/InterchangeTestsModule.h
+++ b/Engine/Plugins/Tests/InterchangeTests/Source/InterchangeTests/Public/InterchangeTestsModule.h
@@ -15,9 +15,6 @@ public:
 	static FInterchangeTestsModule& Get();
 	static bool IsAvailable();
 
-	/** Get a list of json files defining import tests */
-	TArray<FString> GetImportTests() const;
-
 private:
 	virtual void StartupModule() override;
 };
