Improvements and bug fixes to Interchange automation tests.

diff --git a/Engine/Plugins/Tests/InterchangeTests/Source/InterchangeTestEditor/Private/InterchangeTestFunctionLayout.cpp b/Engine/Plugins/Tests/InterchangeTests/Source/InterchangeTestEditor/Private/InterchangeTestFunctionLayout.cpp
--- a/Engine/Plugins/Tests/InterchangeTests/Source/InterchangeTestEditor/Private/InterchangeTestFunctionLayout.cpp
+++ b/Engine/Plugins/Tests/InterchangeTests/Source/InterchangeTestEditor/Private/InterchangeTestFunctionLayout.cpp
@@ -47,6 +47,7 @@ void FInterchangeTestFunctionLayout::CustomizeHeader(TSharedRef<IPropertyHandle>
 	// Cache all the child properties we're interested in manipulating
 
 	AssetClassProperty = StructPropertyHandle->GetChildHandle(GET_MEMBER_NAME_CHECKED(FInterchangeTestFunction, AssetClass));
+	OptionalAssetNameProperty = StructPropertyHandle->GetChildHandle(GET_MEMBER_NAME_CHECKED(FInterchangeTestFunction, OptionalAssetName));
 	CheckFunctionProperty = StructPropertyHandle->GetChildHandle(GET_MEMBER_NAME_CHECKED(FInterchangeTestFunction, CheckFunction));
 	ParametersProperty = StructPropertyHandle->GetChildHandle(GET_MEMBER_NAME_CHECKED(FInterchangeTestFunction, Parameters));
 	PropertyUtilities = StructCustomizationUtils.GetPropertyUtilities();
@@ -75,29 +76,49 @@ void FInterchangeTestFunctionLayout::CustomizeHeader(TSharedRef<IPropertyHandle>
 	]
 	.ValueContent()
 	[
-		SNew(SHorizontalBox)
-		+SHorizontalBox::Slot()
-		.AutoWidth()
+		SNew(SVerticalBox)
+		+SVerticalBox::Slot()
+		.AutoHeight()
 		[
-			SNew(SComboBox<UClass*>)
-			.OptionsSource(&AssetClasses)
-			.OnGenerateWidget(this, &FInterchangeTestFunctionLayout::OnGenerateClassComboWidget)
-			.OnSelectionChanged(this, &FInterchangeTestFunctionLayout::OnClassComboSelectionChanged)
-			.InitiallySelectedItem(InterchangeTestFunction->AssetClass)
+			SNew(SHorizontalBox)
+			+SHorizontalBox::Slot()
+			.AutoWidth()
 			[
-				OnGenerateClassComboWidget(InterchangeTestFunction->AssetClass)
+				SNew(SComboBox<UClass*>)
+				.OptionsSource(&AssetClasses)
+				.OnGenerateWidget(this, &FInterchangeTestFunctionLayout::OnGenerateClassComboWidget)
+				.OnSelectionChanged(this, &FInterchangeTestFunctionLayout::OnClassComboSelectionChanged)
+				.InitiallySelectedItem(InterchangeTestFunction->AssetClass)
+				[
+					OnGenerateClassComboWidget(InterchangeTestFunction->AssetClass)
+				]
+			]
+			+SHorizontalBox::Slot()
+			.AutoWidth()
+			[
+				SNew(SEditableTextBox)
+				.Font(IDetailLayoutBuilder::GetDetailFont())
+				.HintText(LOCTEXT("OptionalAssetName", "Optional Asset Name"))
+				.Text_Lambda([this] { FString Value; OptionalAssetNameProperty->GetValue(Value); return FText::FromString(Value); })
+				.OnTextCommitted_Lambda([this](const FText& Val, ETextCommit::Type TextCommitType) { OptionalAssetNameProperty->SetValue(Val.ToString()); })
+				.ToolTipText(LOCTEXT("OptionalAssetNameTooltip", "Fill out this field with the name of an imported asset which you wish to test, if there are multiple assets of the same type to choose from."))
 			]
 		]
-		+SHorizontalBox::Slot()
-		.AutoWidth()
+		+SVerticalBox::Slot()
+		.AutoHeight()
 		[
-			SNew(SComboBox<UFunction*>)
-			.OptionsSource(&Functions)
-			.OnGenerateWidget(this, &FInterchangeTestFunctionLayout::OnGenerateFunctionComboWidget)
-			.OnSelectionChanged(this, &FInterchangeTestFunctionLayout::OnFunctionComboSelectionChanged)
-			.InitiallySelectedItem(InterchangeTestFunction->CheckFunction)
+			SNew(SHorizontalBox)
+			+ SHorizontalBox::Slot()
+			.AutoWidth()
 			[
-				OnGenerateFunctionComboWidget(InterchangeTestFunction->CheckFunction)
+				SNew(SComboBox<UFunction*>)
+				.OptionsSource(&Functions)
+				.OnGenerateWidget(this, &FInterchangeTestFunctionLayout::OnGenerateFunctionComboWidget)
+				.OnSelectionChanged(this, &FInterchangeTestFunctionLayout::OnFunctionComboSelectionChanged)
+				.InitiallySelectedItem(InterchangeTestFunction->CheckFunction)
+				[
+					OnGenerateFunctionComboWidget(InterchangeTestFunction->CheckFunction)
+				]
 			]
 		]
 	];
diff --git a/Engine/Plugins/Tests/InterchangeTests/Source/InterchangeTestEditor/Private/InterchangeTestFunctionLayout.h b/Engine/Plugins/Tests/InterchangeTests/Source/InterchangeTestEditor/Private/InterchangeTestFunctionLayout.h
--- a/Engine/Plugins/Tests/InterchangeTests/Source/InterchangeTestEditor/Private/InterchangeTestFunctionLayout.h
+++ b/Engine/Plugins/Tests/InterchangeTests/Source/InterchangeTestEditor/Private/InterchangeTestFunctionLayout.h
@@ -49,6 +49,7 @@ private:
 
 	TSharedPtr<IPropertyHandle> StructProperty;
 	TSharedPtr<IPropertyHandle> AssetClassProperty;
+	TSharedPtr<IPropertyHandle> OptionalAssetNameProperty;
 	TSharedPtr<IPropertyHandle> CheckFunctionProperty;
 	TSharedPtr<IPropertyHandle> ParametersProperty;
 
diff --git a/Engine/Plugins/Tests/InterchangeTests/Source/InterchangeTests/Private/ImportTestFunctions/InterchangeResultImportTestFunctions.cpp b/Engine/Plugins/Tests/InterchangeTests/Source/InterchangeTests/Private/ImportTestFunctions/InterchangeResultImportTestFunctions.cpp
--- a/Engine/Plugins/Tests/InterchangeTests/Source/InterchangeTests/Private/ImportTestFunctions/InterchangeResultImportTestFunctions.cpp
+++ b/Engine/Plugins/Tests/InterchangeTests/Source/InterchangeTests/Private/ImportTestFunctions/InterchangeResultImportTestFunctions.cpp
@@ -8,3 +8,29 @@ UClass* UInterchangeResultImportTestFunctions::GetAssociatedAssetType() const
 {
 	return UInterchangeResultsContainer::StaticClass();
 }
+
+
+FInterchangeTestFunctionResult UInterchangeResultImportTestFunctions::CheckIfErrorOrWarningWasGenerated(UInterchangeResultsContainer* ResultsContainer, TSubclassOf<UInterchangeResult> ErrorOrWarningClass)
+{
+	FInterchangeTestFunctionResult Result;
+
+	bool bFoundType = false;
+	for (UInterchangeResult* ResultObject : ResultsContainer->GetResults())
+	{
+		if (ResultObject->IsA(ErrorOrWarningClass))
+		{
+			bFoundType = true;
+			// If we matched the message type, remove it from the results container so it won't flag as an error at the end of the test.
+			// Note: we can do this while iterating through results, because UInterchangeResultsContainer::GetResults passes the array by value,
+			// for thread-safety purposes.
+			ResultsContainer->RemoveResult(ResultObject);
+		}
+	}
+
+	if (!bFoundType)
+	{
+		Result.AddError(FString::Printf(TEXT("No error or warning of class %s was generated."), *ErrorOrWarningClass.Get()->GetName()));
+	}
+
+	return Result;
+}
diff --git a/Engine/Plugins/Tests/InterchangeTests/Source/InterchangeTests/Private/ImportTestFunctions/StaticMeshImportTestFunctions.cpp b/Engine/Plugins/Tests/InterchangeTests/Source/InterchangeTests/Private/ImportTestFunctions/StaticMeshImportTestFunctions.cpp
--- a/Engine/Plugins/Tests/InterchangeTests/Source/InterchangeTests/Private/ImportTestFunctions/StaticMeshImportTestFunctions.cpp
+++ b/Engine/Plugins/Tests/InterchangeTests/Source/InterchangeTests/Private/ImportTestFunctions/StaticMeshImportTestFunctions.cpp
@@ -5,6 +5,7 @@
 #include "Engine/SkeletalMesh.h"
 #include "PhysicsEngine/BodySetup.h"
 #include "StaticMeshAttributes.h"
+#include "ObjectTools.h"
 
 
 UClass* UStaticMeshImportTestFunctions::GetAssociatedAssetType() const
@@ -643,9 +644,9 @@ FInterchangeTestFunctionResult UStaticMeshImportTestFunctions::CheckAgainstGroun
 		}
 	}
 
-	// Don't keep the ground truth asset lying around; mark it as garbage to get rid of it at the end of the tests
-	GroundTruth->RemoveFromRoot();
-	GroundTruth->MarkAsGarbage();
+	// Don't keep the ground truth asset lying around.
+	// @todo: this doesn't work for some reason; find out why.
+//	GroundTruth->ClearFlags(RF_Standalone | RF_WasLoaded | RF_LoadCompleted);
 
 	return Result;
 }
diff --git a/Engine/Plugins/Tests/InterchangeTests/Source/InterchangeTests/Private/InterchangeImportTest.cpp b/Engine/Plugins/Tests/InterchangeTests/Source/InterchangeTests/Private/InterchangeImportTest.cpp
--- a/Engine/Plugins/Tests/InterchangeTests/Source/InterchangeTests/Private/InterchangeImportTest.cpp
+++ b/Engine/Plugins/Tests/InterchangeTests/Source/InterchangeTests/Private/InterchangeImportTest.cpp
@@ -83,11 +83,21 @@ bool FInterchangeImportTest::RunTest(const FString& Path)
 		const FAssetRegistryModule& AssetRegistryModule = FModuleManager::LoadModuleChecked<FAssetRegistryModule>(TEXT("AssetRegistry"));
 
 		AssetRegistryModule.Get().EnumerateAssets(AssetRegistryFilter,
-			[&TestPlans](const FAssetData& AssetData)
+			[this, &TestPlans](const FAssetData& AssetData)
 			{
-				FInterchangeImportTestData& TP = TestPlans.AddDefaulted_GetRef();
-				TP.AssetData = AssetData;
-				TP.TestPlan = CastChecked<UInterchangeImportTestPlan>(AssetData.GetAsset());
+				UInterchangeImportTestPlan* AssetObject = CastChecked<UInterchangeImportTestPlan>(AssetData.GetAsset());
+
+				if (AssetObject->bIsEnabledInAutomationTests)
+				{
+					FInterchangeImportTestData& TP = TestPlans.AddDefaulted_GetRef();
+					TP.AssetData = AssetData;
+					TP.TestPlan = AssetObject;
+				}
+				else
+				{
+					ExecutionInfo.AddWarning(FString::Printf(TEXT("InterchangeImportTestPlan %s is disabled."), *AssetObject->GetName()));
+				}
+
 				return true;
 			}
 		);
@@ -123,7 +133,8 @@ bool FInterchangeImportTest::RunTest(const FString& Path)
 		for (int32 PlanIndex = 0; PlanIndex < TestPlans.Num(); PlanIndex++)
 		{
 			FInterchangeImportTestData& Data = TestPlans[PlanIndex];
-			
+			Data.ImportedAssets.Empty();
+
 			FString AssetName = Data.AssetData.AssetName.ToString().Replace(TEXT("/"), TEXT("_"));
 			Data.DestAssetPackagePath = BasePackagePath / AssetName;
 			Data.DestAssetFilePath = BaseFilePath / AssetName;
@@ -169,7 +180,10 @@ bool FInterchangeImportTest::RunTest(const FString& Path)
 
 						// Fill out list of result objects in the data object.
 						// These are the UObject* results corresponding to imported assets.
-						Data.ResultObjects = Results[PlanIndex]->GetImportedObjects();
+						for (UObject* ImportedObject : Results[PlanIndex]->GetImportedObjects())
+						{
+							Data.ResultObjects.AddUnique(ImportedObject);
+						}
 
 						// Also add the InterchangeResultsContainer to the data so that tests can be run on it
 						// (e.g. to check whether something imported with a specific expected error)
diff --git a/Engine/Plugins/Tests/InterchangeTests/Source/InterchangeTests/Private/InterchangeImportTestStepImport.cpp b/Engine/Plugins/Tests/InterchangeTests/Source/InterchangeTests/Private/InterchangeImportTestStepImport.cpp
--- a/Engine/Plugins/Tests/InterchangeTests/Source/InterchangeTests/Private/InterchangeImportTestStepImport.cpp
+++ b/Engine/Plugins/Tests/InterchangeTests/Source/InterchangeTests/Private/InterchangeImportTestStepImport.cpp
@@ -9,11 +9,14 @@
 #include "Editor.h"
 
 
-UE::Interchange::FAssetImportResultPtr UInterchangeImportTestStepImport::StartStep(const FInterchangeImportTestData& Data)
+UE::Interchange::FAssetImportResultPtr UInterchangeImportTestStepImport::StartStep(FInterchangeImportTestData& Data)
 {
 	// Empty the destination folder here if requested
 	if (bEmptyDestinationFolderPriorToImport)
 	{
+		Data.ResultObjects.Empty();
+		Data.ImportedAssets.Empty();
+
 		const bool bRequireExists = true;
 		const bool bDeleteRecursively = true;
 		IFileManager::Get().DeleteDirectory(*Data.DestAssetFilePath, bRequireExists, bDeleteRecursively);
@@ -61,11 +64,6 @@ FTestStepResults UInterchangeImportTestStepImport::FinishStep(FInterchangeImport
 				SaveArgs);
 		}
 
-		if ((GEditor != nullptr) && (GEditor->Trans != nullptr))
-		{
-			GEditor->Trans->Reset(FText::FromString("Discard undo history during Automation testing."));
-		}
-
 		// Then rename original objects and their packages, and mark as garbage
 		for (const FAssetData& AssetData : Data.ImportedAssets)
 		{
@@ -91,11 +89,12 @@ FTestStepResults UInterchangeImportTestStepImport::FinishStep(FInterchangeImport
 			PackageObject->Rename(*(PackageObject->GetName() + TEXT("_TRASH")), nullptr, RenameFlags);
 			PackageObject->RemoveFromRoot();
 			PackageObject->MarkAsGarbage();
-		}
 
+			// Remove the old version of the asset object from the results
+			Data.ResultObjects.Remove(AssetObject);
+		}
 
 		// Now reload
-		Data.ResultObjects.Reset();
 		for (const FAssetData& AssetData : Data.ImportedAssets)
 		{
 			check(!AssetData.IsAssetLoaded());
diff --git a/Engine/Plugins/Tests/InterchangeTests/Source/InterchangeTests/Private/InterchangeImportTestStepReimport.cpp b/Engine/Plugins/Tests/InterchangeTests/Source/InterchangeTests/Private/InterchangeImportTestStepReimport.cpp
--- a/Engine/Plugins/Tests/InterchangeTests/Source/InterchangeTests/Private/InterchangeImportTestStepReimport.cpp
+++ b/Engine/Plugins/Tests/InterchangeTests/Source/InterchangeTests/Private/InterchangeImportTestStepReimport.cpp
@@ -6,7 +6,7 @@
 #include "UObject/SavePackage.h"
 
 
-UE::Interchange::FAssetImportResultPtr UInterchangeImportTestStepReimport::StartStep(const FInterchangeImportTestData& Data)
+UE::Interchange::FAssetImportResultPtr UInterchangeImportTestStepReimport::StartStep(FInterchangeImportTestData& Data)
 {
 	// Find the object we wish to reimport
 	TArray<UObject*> PotentialObjectsToReimport;
@@ -99,10 +99,11 @@ FTestStepResults UInterchangeImportTestStepReimport::FinishStep(FInterchangeImpo
 			AssetObject->RemoveFromRoot();
 			AssetObject->MarkAsGarbage();
 			PackageObject->MarkAsGarbage();
+
+			Data.ResultObjects.Remove(AssetObject);
 		}
 
 		// Now reload
-		Data.ResultObjects.Reset();
 		for (const FAssetData& AssetData : Data.ImportedAssets)
 		{
 			check(!AssetData.IsAssetLoaded());
diff --git a/Engine/Plugins/Tests/InterchangeTests/Source/InterchangeTests/Private/InterchangeTestFunction.cpp b/Engine/Plugins/Tests/InterchangeTests/Source/InterchangeTests/Private/InterchangeTestFunction.cpp
--- a/Engine/Plugins/Tests/InterchangeTests/Source/InterchangeTests/Private/InterchangeTestFunction.cpp
+++ b/Engine/Plugins/Tests/InterchangeTests/Source/InterchangeTests/Private/InterchangeTestFunction.cpp
@@ -220,7 +220,12 @@ FInterchangeTestFunctionResult FInterchangeTestFunction::Invoke(const TArray<UOb
 	{
 		if (Asset->GetClass()->IsChildOf(ExpectedClass))
 		{
-			AssetsToPass.Add(Asset);
+			// If the asset is of the expected class, add it if the name matches the supplied optional asset name,
+			// or if we are not checking a specific name.
+			if (OptionalAssetName.IsEmpty() || OptionalAssetName == Asset->GetName())
+			{
+				AssetsToPass.Add(Asset);
+			}
 		}
 	}
 
@@ -235,6 +240,18 @@ FInterchangeTestFunctionResult FInterchangeTestFunction::Invoke(const TArray<UOb
 	}
 	else
 	{
+		// If there was an asset name, but no matches, return with an error
+		if (!OptionalAssetName.IsEmpty() && AssetsToPass.Num() == 0)
+		{
+			FInterchangeTestFunctionResult Result;
+			Result.AddError(FString::Printf(
+				TEXT("No asset with the name '%s' was found"),
+				*OptionalAssetName)
+			);
+
+			return Result;
+		}
+
 		// If there is not exactly one entry in the array, the result is ambiguous, so return with an error
 		if (AssetsToPass.Num() != 1)
 		{
diff --git a/Engine/Plugins/Tests/InterchangeTests/Source/InterchangeTests/Public/ImportTestFunctions/InterchangeResultImportTestFunctions.h b/Engine/Plugins/Tests/InterchangeTests/Source/InterchangeTests/Public/ImportTestFunctions/InterchangeResultImportTestFunctions.h
--- a/Engine/Plugins/Tests/InterchangeTests/Source/InterchangeTests/Public/ImportTestFunctions/InterchangeResultImportTestFunctions.h
+++ b/Engine/Plugins/Tests/InterchangeTests/Source/InterchangeTests/Public/ImportTestFunctions/InterchangeResultImportTestFunctions.h
@@ -4,10 +4,10 @@
 
 #include "ImportTestFunctionsBase.h"
 #include "InterchangeTestFunction.h"
+#include "InterchangeResultsContainer.h"
 #include "InterchangeResultImportTestFunctions.generated.h"
 
 struct FInterchangeTestFunctionResult;
-class UInterchangeResultsContainer;
 
 
 UCLASS()
@@ -19,4 +19,8 @@ public:
 
 	// UImportTestFunctionsBase interface
 	virtual UClass* GetAssociatedAssetType() const override;
+
+	/** Check whether the specified InterchangeResult was emitted during import */
+	UFUNCTION(Exec)
+	static FInterchangeTestFunctionResult CheckIfErrorOrWarningWasGenerated(UInterchangeResultsContainer* ResultsContainer, TSubclassOf<UInterchangeResult> ErrorOrWarningClass);
 };
diff --git a/Engine/Plugins/Tests/InterchangeTests/Source/InterchangeTests/Public/InterchangeImportTestPlan.h b/Engine/Plugins/Tests/InterchangeTests/Source/InterchangeTests/Public/InterchangeImportTestPlan.h
--- a/Engine/Plugins/Tests/InterchangeTests/Source/InterchangeTests/Public/InterchangeImportTestPlan.h
+++ b/Engine/Plugins/Tests/InterchangeTests/Source/InterchangeTests/Public/InterchangeImportTestPlan.h
@@ -27,11 +27,14 @@ public:
 	void WriteToJson(const FString& Filename);
 
 public:
-
 	/** Set of steps to perform to carry out this test plan */
 	UPROPERTY(EditAnywhere, Instanced, Category = Definition, Meta = (DisplayPriority = 0))
 	TArray<TObjectPtr<UInterchangeImportTestStepBase>> Steps;
 
+	/** Whether or not this test plan is currently enabled */
+	UPROPERTY(EditAnywhere, Category = Activation)
+	bool bIsEnabledInAutomationTests = true;
+
 	/** Click here to immediately run this single test through the automation framework */
 	UFUNCTION(BlueprintCallable, CallInEditor, Category = Run, Meta = (DisplayPriority = 1))
 	void RunThisTest();
diff --git a/Engine/Plugins/Tests/InterchangeTests/Source/InterchangeTests/Public/InterchangeImportTestStepBase.h b/Engine/Plugins/Tests/InterchangeTests/Source/InterchangeTests/Public/InterchangeImportTestStepBase.h
--- a/Engine/Plugins/Tests/InterchangeTests/Source/InterchangeTests/Public/InterchangeImportTestStepBase.h
+++ b/Engine/Plugins/Tests/InterchangeTests/Source/InterchangeTests/Public/InterchangeImportTestStepBase.h
@@ -30,7 +30,7 @@ public:
 	TArray<FInterchangeTestFunction> Tests;
 
 public:
-	virtual UE::Interchange::FAssetImportResultPtr StartStep(const FInterchangeImportTestData& Data) PURE_VIRTUAL(UInterchangeImportTestStepBase::StartStep, return {}; );
+	virtual UE::Interchange::FAssetImportResultPtr StartStep(FInterchangeImportTestData& Data) PURE_VIRTUAL(UInterchangeImportTestStepBase::StartStep, return {}; );
 	virtual FTestStepResults FinishStep(FInterchangeImportTestData& Data, FAutomationTestExecutionInfo& ExecutionInfo) PURE_VIRTUAL(UInterchangeImportTestStepBase::FinishStep, return {}; );
 	virtual FString GetContextString() const PURE_VIRTUAL(UInterchangeImportTestStepBase::GetContextString, return {}; );
 
diff --git a/Engine/Plugins/Tests/InterchangeTests/Source/InterchangeTests/Public/InterchangeImportTestStepImport.h b/Engine/Plugins/Tests/InterchangeTests/Source/InterchangeTests/Public/InterchangeImportTestStepImport.h
--- a/Engine/Plugins/Tests/InterchangeTests/Source/InterchangeTests/Public/InterchangeImportTestStepImport.h
+++ b/Engine/Plugins/Tests/InterchangeTests/Source/InterchangeTests/Public/InterchangeImportTestStepImport.h
@@ -34,7 +34,7 @@ public:
 
 public:
 	// UInterchangeImportTestStepBase interface
-	virtual UE::Interchange::FAssetImportResultPtr StartStep(const FInterchangeImportTestData& Data) override;
+	virtual UE::Interchange::FAssetImportResultPtr StartStep(FInterchangeImportTestData& Data) override;
 	virtual FTestStepResults FinishStep(FInterchangeImportTestData& Data, FAutomationTestExecutionInfo& ExecutionInfo) override;
 	virtual FString GetContextString() const override;
 };
diff --git a/Engine/Plugins/Tests/InterchangeTests/Source/InterchangeTests/Public/InterchangeImportTestStepReimport.h b/Engine/Plugins/Tests/InterchangeTests/Source/InterchangeTests/Public/InterchangeImportTestStepReimport.h
--- a/Engine/Plugins/Tests/InterchangeTests/Source/InterchangeTests/Public/InterchangeImportTestStepReimport.h
+++ b/Engine/Plugins/Tests/InterchangeTests/Source/InterchangeTests/Public/InterchangeImportTestStepReimport.h
@@ -34,7 +34,7 @@ public:
 
 public:
 	// UInterchangeImportTestStepBase interface
-	virtual UE::Interchange::FAssetImportResultPtr StartStep(const FInterchangeImportTestData& Data) override;
+	virtual UE::Interchange::FAssetImportResultPtr StartStep(FInterchangeImportTestData& Data) override;
 	virtual FTestStepResults FinishStep(FInterchangeImportTestData& Data, FAutomationTestExecutionInfo& ExecutionInfo) override;
 	virtual FString GetContextString() const override;
 };
diff --git a/Engine/Plugins/Tests/InterchangeTests/Source/InterchangeTests/Public/InterchangeTestFunction.h b/Engine/Plugins/Tests/InterchangeTests/Source/InterchangeTests/Public/InterchangeTestFunction.h
--- a/Engine/Plugins/Tests/InterchangeTests/Source/InterchangeTests/Public/InterchangeTestFunction.h
+++ b/Engine/Plugins/Tests/InterchangeTests/Source/InterchangeTests/Public/InterchangeTestFunction.h
@@ -60,6 +60,10 @@ public:
 	UPROPERTY(EditAnywhere, Category = Test)
 	TSubclassOf<UObject> AssetClass;
 
+	/** Optional name of the asset to test, if there are various contenders */
+	UPROPERTY(EditAnywhere, Category = Test)
+	FString OptionalAssetName;
+
 	/** A function to be called to determine whether the result is correct */
 	UPROPERTY(EditAnywhere, Category = Test)
 	TObjectPtr<UFunction> CheckFunction = nullptr;
diff --git a/Engine/Source/Runtime/Experimental/Interchange/Core/Public/InterchangeResultsContainer.h b/Engine/Source/Runtime/Experimental/Interchange/Core/Public/InterchangeResultsContainer.h
--- a/Engine/Source/Runtime/Experimental/Interchange/Core/Public/InterchangeResultsContainer.h
+++ b/Engine/Source/Runtime/Experimental/Interchange/Core/Public/InterchangeResultsContainer.h
@@ -70,6 +70,15 @@ public:
 		return Results;
 	}
 
+	/**
+	 * Removes the given UInterchangeResult from the container.
+	 */
+	void RemoveResult(UInterchangeResult* Item)
+	{
+		FScopeLock ScopeLock(&Lock);
+		Results.Remove(Item);
+	}
+
 private:
 
 	mutable FCriticalSection Lock;
