A few small automation bugfixes.

diff --git a/Engine/Plugins/Tests/InterchangeTests/Source/InterchangeTests/Private/InterchangeImportTest.cpp b/Engine/Plugins/Tests/InterchangeTests/Source/InterchangeTests/Private/InterchangeImportTest.cpp
--- a/Engine/Plugins/Tests/InterchangeTests/Source/InterchangeTests/Private/InterchangeImportTest.cpp
+++ b/Engine/Plugins/Tests/InterchangeTests/Source/InterchangeTests/Private/InterchangeImportTest.cpp
@@ -189,6 +189,18 @@ bool FInterchangeImportTest::RunTest(const FString& Path)
 						// (e.g. to check whether something imported with a specific expected error)
 						Data.InterchangeResults = Results[PlanIndex]->GetResults();
 
+						// Fill out list of imported assets as FAssetData
+						for (UObject* Object : Results[PlanIndex]->GetImportedObjects())
+						{
+							Data.ImportedAssets.Emplace(Object);
+						}
+
+						// We don't need the Interchange results object any more, as we've already taken everything from it that we need.
+						// If we don't reset it, it will hold on to the trashed versions of the imported assets during GC.
+						Results[PlanIndex].Reset();
+
+						Data.TestPlan->Steps[StepIndex]->FinishStep(Data, ExecutionInfo);
+
 						// Populate the automation test execution info with the interchange import results
 						for (UInterchangeResult* Result : Data.InterchangeResults->GetResults())
 						{
@@ -196,6 +208,7 @@ bool FInterchangeImportTest::RunTest(const FString& Path)
 							{
 							case EInterchangeResultType::Error:
 								ExecutionInfo.AddError(Result->GetText().ToString());
+								bSuccess = false;
 								break;
 
 							case EInterchangeResultType::Warning:
@@ -204,24 +217,18 @@ bool FInterchangeImportTest::RunTest(const FString& Path)
 							}
 						}
 
-						// Fill out list of imported assets as FAssetData
-						for (UObject* Object : Results[PlanIndex]->GetImportedObjects())
-						{
-							Data.ImportedAssets.Emplace(Object);
-						}
-
-						// We don't need the Interchange results object any more, as we've already taken everything from it that we need.
-						// If we don't reset it, it will hold on to the trashed versions of the imported assets during GC.
-						Results[PlanIndex].Reset();
-
-						Data.TestPlan->Steps[StepIndex]->FinishStep(Data, ExecutionInfo);
-
 						ExecutionInfo.PopContext();
 					}
 				}
 			}
 
-			// Collect garbage between every step, so that 
+			// Collect garbage between every step, so that we remove renamed packages which come from the save+reload operation
+			// Note we also reset the transaction buffer here to stop it from holding onto references which would prevent garbage collection.
+			// @todo: not really a big fan of this; is there a better way of just disabling transactions?
+			if ((GEditor != nullptr) && (GEditor->Trans != nullptr))
+			{
+				GEditor->Trans->Reset(FText::FromString("Discard undo history during FBX Automation testing."));
+			}
 			CollectGarbage(GARBAGE_COLLECTION_KEEPFLAGS);
 		}
 
@@ -237,5 +244,5 @@ bool FInterchangeImportTest::RunTest(const FString& Path)
 	bool bShowConfirmation = false;
 	ObjectTools::ForceDeleteObjects(ObjectsToDelete, bShowConfirmation);
 
-	return true;
+	return bSuccess;
 }
diff --git a/Engine/Plugins/Tests/InterchangeTests/Source/InterchangeTests/Private/InterchangeImportTestStepBase.cpp b/Engine/Plugins/Tests/InterchangeTests/Source/InterchangeTests/Private/InterchangeImportTestStepBase.cpp
--- a/Engine/Plugins/Tests/InterchangeTests/Source/InterchangeTests/Private/InterchangeImportTestStepBase.cpp
+++ b/Engine/Plugins/Tests/InterchangeTests/Source/InterchangeTests/Private/InterchangeImportTestStepBase.cpp
@@ -15,7 +15,7 @@ bool UInterchangeImportTestStepBase::PerformTests(FInterchangeImportTestData& Da
 
 	for (FInterchangeTestFunction& Test : Tests)
 	{
-		FInterchangeTestFunctionResult Result = Test.Invoke(Data.ResultObjects);
+		FInterchangeTestFunctionResult Result = Test.Invoke(ResultObjects);
 
 		for (const FString& Warning : Result.GetWarnings())
 		{
diff --git a/Engine/Plugins/Tests/InterchangeTests/Source/InterchangeTests/Private/InterchangeImportTestStepReimport.cpp b/Engine/Plugins/Tests/InterchangeTests/Source/InterchangeTests/Private/InterchangeImportTestStepReimport.cpp
--- a/Engine/Plugins/Tests/InterchangeTests/Source/InterchangeTests/Private/InterchangeImportTestStepReimport.cpp
+++ b/Engine/Plugins/Tests/InterchangeTests/Source/InterchangeTests/Private/InterchangeImportTestStepReimport.cpp
@@ -89,15 +89,21 @@ FTestStepResults UInterchangeImportTestStepReimport::FinishStep(FInterchangeImpo
 			check(PackageObject);
 			check(PackageObject == AssetData.GetPackage());
 
+			// Mark all objects in the package as garbage, and remove the standalone flag, so that GC can remove the package later
+			TArray<UObject*> ObjectsInPackage;
+			GetObjectsWithPackage(PackageObject, ObjectsInPackage, true);
+			for (UObject* ObjectInPackage : ObjectsInPackage)
+			{
+				ObjectInPackage->ClearFlags(RF_Standalone | RF_Public);
+				ObjectInPackage->MarkAsGarbage();
+			}
+
 			// Renaming the original objects avoids having to do a GC sweep here.
 			// Any existing references to them will be retained but irrelevant.
 			// Then the new object can be loaded in their place, as if it were being loaded for the first time.
 			const ERenameFlags RenameFlags = REN_ForceNoResetLoaders | REN_DontCreateRedirectors | REN_NonTransactional | REN_DoNotDirty;
-			AssetObject->Rename(*(AssetObject->GetName() + TEXT("_TRASH")), PackageObject, RenameFlags);
 			PackageObject->Rename(*(PackageObject->GetName() + TEXT("_TRASH")), nullptr, RenameFlags);
-			AssetObject->ClearFlags(RF_Standalone | RF_Public);
-			AssetObject->RemoveFromRoot();
-			AssetObject->MarkAsGarbage();
+			PackageObject->RemoveFromRoot();
 			PackageObject->MarkAsGarbage();
 
 			Data.ResultObjects.Remove(AssetObject);
