Disable Interchange automation tests on build machines, as we currently are unable to keep InterchangeWorker.exe (a necessary dependency) up to date.

diff --git a/Engine/Plugins/Tests/InterchangeTests/Source/InterchangeTests/Private/InterchangeImportTest.cpp b/Engine/Plugins/Tests/InterchangeTests/Source/InterchangeTests/Private/InterchangeImportTest.cpp
--- a/Engine/Plugins/Tests/InterchangeTests/Source/InterchangeTests/Private/InterchangeImportTest.cpp
+++ b/Engine/Plugins/Tests/InterchangeTests/Source/InterchangeTests/Private/InterchangeImportTest.cpp
@@ -29,6 +29,13 @@ IMPLEMENT_COMPLEX_AUTOMATION_TEST(FInterchangeImportTest, "Editor.Interchange.Im
 
 void FInterchangeImportTest::GetTests(TArray<FString>& OutBeautifiedNames, TArray<FString>& OutTestCommands) const
 {
+	// For now, we can't run Interchange automation tests under CIS as we have no way of keeping InterchangeWorker.exe up to date.
+	// @todo: find a solution to this
+	if (GIsBuildMachine)
+	{
+		return;
+	}
+
 	const FAssetRegistryModule& AssetRegistryModule = FModuleManager::LoadModuleChecked<FAssetRegistryModule>(TEXT("AssetRegistry"));
 	TArray<FAssetData> AllTestPlans;
 	AssetRegistryModule.Get().GetAssetsByClass(UInterchangeImportTestPlan::StaticClass()->GetFName(), AllTestPlans, true);
@@ -57,6 +64,13 @@ void FInterchangeImportTest::GetTests(TArray<FString>& OutBeautifiedNames, TArra
 
 bool FInterchangeImportTest::RunTest(const FString& Path)
 {
+	// For now, we can't run Interchange automation tests under CIS as we have no way of keeping InterchangeWorker.exe up to date.
+	// @todo: find a solution to this
+	if (GIsBuildMachine)
+	{
+		return true;
+	}
+
 	// Determine the test plan assets within the given path which will be run in parallel
 
 	TArray<FInterchangeImportTestData> TestPlans;
