For now, disable Interchange automation tests on Mac, as there's an issue with the InterchangeWorker process spawning on Mac which needs to be investigated. This is just a temporary fix to stop CIS failures.
Also, null out results container pointer after testing so that it is not pointing to junk after garbage collection.

diff --git a/Engine/Plugins/Tests/InterchangeTests/Source/InterchangeTests/Private/InterchangeImportTest.cpp b/Engine/Plugins/Tests/InterchangeTests/Source/InterchangeTests/Private/InterchangeImportTest.cpp
--- a/Engine/Plugins/Tests/InterchangeTests/Source/InterchangeTests/Private/InterchangeImportTest.cpp
+++ b/Engine/Plugins/Tests/InterchangeTests/Source/InterchangeTests/Private/InterchangeImportTest.cpp
@@ -30,12 +30,11 @@ IMPLEMENT_COMPLEX_AUTOMATION_TEST(FInterchangeImportTest, "Editor.Interchange.Im
 
 void FInterchangeImportTest::GetTests(TArray<FString>& OutBeautifiedNames, TArray<FString>& OutTestCommands) const
 {
-	// For now, we can't run Interchange automation tests under CIS as we have no way of keeping InterchangeWorker.exe up to date.
+	// For now, we can't run Interchange automation tests on Mac, as there's a problem with the InterchangeWorker. For now, do nothing on Mac.
 	// @todo: find a solution to this
-	if (GIsBuildMachine)
-	{
-		return;
-	}
+#if PLATFORM_MAC
+	return;
+#endif
 
 	const FAssetRegistryModule& AssetRegistryModule = FModuleManager::LoadModuleChecked<FAssetRegistryModule>(TEXT("AssetRegistry"));
 	TArray<FAssetData> AllTestPlans;
@@ -65,12 +64,11 @@ void FInterchangeImportTest::GetTests(TArray<FString>& OutBeautifiedNames, TArra
 
 bool FInterchangeImportTest::RunTest(const FString& Path)
 {
-	// For now, we can't run Interchange automation tests under CIS as we have no way of keeping InterchangeWorker.exe up to date.
+	// For now, we can't run Interchange automation tests on Mac, as there's a problem with the InterchangeWorker. For now, do nothing on Mac.
 	// @todo: find a solution to this
-	if (GIsBuildMachine)
-	{
-		return true;
-	}
+#if PLATFORM_MAC
+	return true;
+#endif
 
 	// Determine the test plan assets within the given path which will be run in parallel
 
@@ -232,6 +230,9 @@ bool FInterchangeImportTest::RunTest(const FString& Path)
 							}
 						}
 
+						// Finished with the interchange results - null it so it will be GC'd later
+						Data.InterchangeResults = nullptr;
+
 						ExecutionInfo.PopContext();
 					}
 				}
