Disable Linux tests here as well, same reason as UE-146411

diff --git a/Engine/Plugins/Tests/InterchangeTests/Source/InterchangeTests/Private/InterchangeImportTest.cpp b/Engine/Plugins/Tests/InterchangeTests/Source/InterchangeTests/Private/InterchangeImportTest.cpp
--- a/Engine/Plugins/Tests/InterchangeTests/Source/InterchangeTests/Private/InterchangeImportTest.cpp
+++ b/Engine/Plugins/Tests/InterchangeTests/Source/InterchangeTests/Private/InterchangeImportTest.cpp
@@ -32,7 +32,7 @@ void FInterchangeImportTest::GetTests(TArray<FString>& OutBeautifiedNames, TArra
 {
 	// For now, we can't run Interchange automation tests on Mac, as there's a problem with the InterchangeWorker. For now, do nothing on Mac.
 	// @todo: find a solution to this
-#if PLATFORM_MAC
+#if PLATFORM_MAC || PLATFORM_LINUX
 	return;
 #endif
 
@@ -66,7 +66,7 @@ bool FInterchangeImportTest::RunTest(const FString& Path)
 {
 	// For now, we can't run Interchange automation tests on Mac, as there's a problem with the InterchangeWorker. For now, do nothing on Mac.
 	// @todo: find a solution to this
-#if PLATFORM_MAC
+#if PLATFORM_MAC || PLATFORM_LINUX
 	return true;
 #endif
 
