InterchangeTest framework: Added an import into level option on the import test step
- The reimport step doesn't support import into level for now since reimport for scenes isn't implemented yet.
- Can only count the number of actors imported for now.

diff --git a/Engine/Plugins/Tests/InterchangeTests/Source/InterchangeTests/Private/ImportTestFunctions/ActorImportTestFunctions.cpp b/Engine/Plugins/Tests/InterchangeTests/Source/InterchangeTests/Private/ImportTestFunctions/ActorImportTestFunctions.cpp
new file mode 100644
--- /dev/null
+++ b/Engine/Plugins/Tests/InterchangeTests/Source/InterchangeTests/Private/ImportTestFunctions/ActorImportTestFunctions.cpp
@@ -0,0 +1,21 @@
+// Copyright Epic Games, Inc. All Rights Reserved.
+
+#include "ImportTestFunctions/ActorImportTestFunctions.h"
+#include "GameFramework/Actor.h"
+
+
+UClass* UActorImportTestFunctions::GetAssociatedAssetType() const
+{
+	return AActor::StaticClass();
+}
+
+FInterchangeTestFunctionResult UActorImportTestFunctions::CheckImportedActorCount(const TArray<AActor*>& Actors, int32 ExpectedNumberOfImportedActors)
+{
+	FInterchangeTestFunctionResult Result;
+	if (Actors.Num() != ExpectedNumberOfImportedActors)
+	{
+		Result.AddError(FString::Printf(TEXT("Expected %d actors, imported %d."), ExpectedNumberOfImportedActors, Actors.Num()));
+	}
+
+	return Result;
+}
diff --git a/Engine/Plugins/Tests/InterchangeTests/Source/InterchangeTests/Private/InterchangeImportTest.cpp b/Engine/Plugins/Tests/InterchangeTests/Source/InterchangeTests/Private/InterchangeImportTest.cpp
--- a/Engine/Plugins/Tests/InterchangeTests/Source/InterchangeTests/Private/InterchangeImportTest.cpp
+++ b/Engine/Plugins/Tests/InterchangeTests/Source/InterchangeTests/Private/InterchangeImportTest.cpp
@@ -128,7 +128,7 @@ bool FInterchangeImportTest::RunTest(const FString& Path)
 	const bool bDeleteRecursively = true;
 	IFileManager::Get().DeleteDirectory(*BaseFilePath, bRequireExists, bDeleteRecursively);
 
-	TArray<UE::Interchange::FAssetImportResultPtr> Results;
+	TArray<TTuple<UE::Interchange::FAssetImportResultPtr, UE::Interchange::FSceneImportResultPtr>> Results;
 
 	int32 StepIndex = 0;
 	bool bFoundValidSteps = true;
@@ -172,9 +172,14 @@ bool FInterchangeImportTest::RunTest(const FString& Path)
 
 			for (int32 PlanIndex = 0; PlanIndex < TestPlans.Num(); PlanIndex++)
 			{
-				if (Results[PlanIndex])
+				if (Results[PlanIndex].Get<0>())
 				{
-					Results[PlanIndex]->WaitUntilDone();
+					Results[PlanIndex].Get<0>()->WaitUntilDone();
+				}
+
+				if (Results[PlanIndex].Get<1>())
+				{
+					Results[PlanIndex].Get<1>()->WaitUntilDone();
 				}
 			}
 
@@ -193,24 +198,44 @@ bool FInterchangeImportTest::RunTest(const FString& Path)
 
 						// Fill out list of result objects in the data object.
 						// These are the UObject* results corresponding to imported assets.
-						for (UObject* ImportedObject : Results[PlanIndex]->GetImportedObjects())
+						if (Results[PlanIndex].Get<0>())
 						{
-							Data.ResultObjects.AddUnique(ImportedObject);
-						}
+							for (UObject* ImportedObject : Results[PlanIndex].Get<0>()->GetImportedObjects())
+							{
+								Data.ResultObjects.AddUnique(ImportedObject);
+							}
 
-						// Also add the InterchangeResultsContainer to the data so that tests can be run on it
-						// (e.g. to check whether something imported with a specific expected error)
-						Data.InterchangeResults = Results[PlanIndex]->GetResults();
+							// Also add the InterchangeResultsContainer to the data so that tests can be run on it
+							// (e.g. to check whether something imported with a specific expected error)
+							Data.InterchangeResults = Results[PlanIndex].Get<0>()->GetResults();
 
-						// Fill out list of imported assets as FAssetData
-						for (UObject* Object : Results[PlanIndex]->GetImportedObjects())
+							// Fill out list of imported assets as FAssetData
+							for (UObject* Object : Results[PlanIndex].Get<0>()->GetImportedObjects())
+							{
+								Data.ImportedAssets.Emplace(Object);
+							}
+						}
+
+						if (Results[PlanIndex].Get<1>())
 						{
-							Data.ImportedAssets.Emplace(Object);
+							for (UObject* ImportedObject : Results[PlanIndex].Get<1>()->GetImportedObjects())
+							{
+								Data.ResultObjects.AddUnique(ImportedObject);
+							}
+
+							if (Data.InterchangeResults)
+							{
+								Data.InterchangeResults->Append(Results[PlanIndex].Get<1>()->GetResults());
+							}
+							else
+							{
+								Data.InterchangeResults = Results[PlanIndex].Get<1>()->GetResults();
+							}
 						}
 
 						// We don't need the Interchange results object any more, as we've already taken everything from it that we need.
 						// If we don't reset it, it will hold on to the trashed versions of the imported assets during GC.
-						Results[PlanIndex].Reset();
+						Results[PlanIndex] = {};
 
 						Data.TestPlan->Steps[StepIndex]->FinishStep(Data, ExecutionInfo);
 
@@ -254,10 +279,23 @@ bool FInterchangeImportTest::RunTest(const FString& Path)
 	TArray<UObject*> ObjectsToDelete;
 	for (FInterchangeImportTestData& Data : TestPlans)
 	{
-		ObjectsToDelete.Append(Data.ResultObjects);
+		ObjectsToDelete.Reserve(ObjectsToDelete.Max() + Data.ResultObjects.Num());
+
+		for (UObject* ResultObject : Data.ResultObjects)
+		{
+			if (AActor* Actor = Cast<AActor>(ResultObject))
+			{
+				constexpr bool bShouldModifyLevel = false;
+				Actor->GetWorld()->RemoveActor(Actor, bShouldModifyLevel);
+			}
+			else
+			{
+				ObjectsToDelete.Add(ResultObject);
+			}
+		}
 	}
 
-	bool bShowConfirmation = false;
+	constexpr bool bShowConfirmation = false;
 	ObjectTools::ForceDeleteObjects(ObjectsToDelete, bShowConfirmation);
 
 	return bSuccess;
diff --git a/Engine/Plugins/Tests/InterchangeTests/Source/InterchangeTests/Private/InterchangeImportTestStepImport.cpp b/Engine/Plugins/Tests/InterchangeTests/Source/InterchangeTests/Private/InterchangeImportTestStepImport.cpp
--- a/Engine/Plugins/Tests/InterchangeTests/Source/InterchangeTests/Private/InterchangeImportTestStepImport.cpp
+++ b/Engine/Plugins/Tests/InterchangeTests/Source/InterchangeTests/Private/InterchangeImportTestStepImport.cpp
@@ -9,7 +9,7 @@
 #include "Editor.h"
 
 
-UE::Interchange::FAssetImportResultPtr UInterchangeImportTestStepImport::StartStep(FInterchangeImportTestData& Data)
+TTuple<UE::Interchange::FAssetImportResultPtr, UE::Interchange::FSceneImportResultPtr> UInterchangeImportTestStepImport::StartStep(FInterchangeImportTestData& Data)
 {
 	// Empty the destination folder here if requested
 	if (bEmptyDestinationFolderPriorToImport)
@@ -35,7 +35,16 @@ UE::Interchange::FAssetImportResultPtr UInterchangeImportTestStepImport::StartSt
 	Params.bIsAutomated = true;
 
 	UInterchangeManager& InterchangeManager = UInterchangeManager::GetInterchangeManager();
-	return InterchangeManager.ImportAssetAsync(Data.DestAssetPackagePath, ScopedSourceData.GetSourceData(), Params);
+
+	if (bImportIntoLevel)
+	{
+		return InterchangeManager.ImportSceneAsync(Data.DestAssetPackagePath, ScopedSourceData.GetSourceData(), Params);
+	}
+	else
+	{
+		UE::Interchange::FAssetImportResultPtr AssetImportResult = InterchangeManager.ImportAssetAsync(Data.DestAssetPackagePath, ScopedSourceData.GetSourceData(), Params);
+		return {AssetImportResult, nullptr};
+	}
 }
 
 
@@ -44,7 +53,8 @@ FTestStepResults UInterchangeImportTestStepImport::FinishStep(FInterchangeImport
 	FTestStepResults Results;
 
 	// If we need to perform a save and fresh reload of everything we imported, do it here
-	if (bSaveThenReloadImportedAssets)
+	// Incompatible with bImportIntoLevel for now, since we'd need to patch all references to the assets
+	if (bSaveThenReloadImportedAssets && !bImportIntoLevel)
 	{
 		// First save
 		for (const FAssetData& AssetData : Data.ImportedAssets)
diff --git a/Engine/Plugins/Tests/InterchangeTests/Source/InterchangeTests/Private/InterchangeImportTestStepReimport.cpp b/Engine/Plugins/Tests/InterchangeTests/Source/InterchangeTests/Private/InterchangeImportTestStepReimport.cpp
--- a/Engine/Plugins/Tests/InterchangeTests/Source/InterchangeTests/Private/InterchangeImportTestStepReimport.cpp
+++ b/Engine/Plugins/Tests/InterchangeTests/Source/InterchangeTests/Private/InterchangeImportTestStepReimport.cpp
@@ -6,7 +6,7 @@
 #include "UObject/SavePackage.h"
 
 
-UE::Interchange::FAssetImportResultPtr UInterchangeImportTestStepReimport::StartStep(FInterchangeImportTestData& Data)
+TTuple<UE::Interchange::FAssetImportResultPtr, UE::Interchange::FSceneImportResultPtr> UInterchangeImportTestStepReimport::StartStep(FInterchangeImportTestData& Data)
 {
 	// Find the object we wish to reimport
 	TArray<UObject*> PotentialObjectsToReimport;
@@ -40,7 +40,7 @@ UE::Interchange::FAssetImportResultPtr UInterchangeImportTestStepReimport::Start
 
 	if (AssetToReimport == nullptr)
 	{
-		return UE::Interchange::FAssetImportResultPtr();
+		return {nullptr, nullptr};
 	}
 
 	// Start the Interchange import
@@ -51,7 +51,8 @@ UE::Interchange::FAssetImportResultPtr UInterchangeImportTestStepReimport::Start
 	Params.ReimportAsset = AssetToReimport;
 
 	UInterchangeManager& InterchangeManager = UInterchangeManager::GetInterchangeManager();
-	return InterchangeManager.ImportAssetAsync(Data.DestAssetPackagePath, ScopedSourceData.GetSourceData(), Params);
+	UE::Interchange::FAssetImportResultPtr AssetImportResult = InterchangeManager.ImportAssetAsync(Data.DestAssetPackagePath, ScopedSourceData.GetSourceData(), Params);
+	return {AssetImportResult, nullptr};
 }
 
 
diff --git a/Engine/Plugins/Tests/InterchangeTests/Source/InterchangeTests/Public/ImportTestFunctions/ActorImportTestFunctions.h b/Engine/Plugins/Tests/InterchangeTests/Source/InterchangeTests/Public/ImportTestFunctions/ActorImportTestFunctions.h
new file mode 100644
--- /dev/null
+++ b/Engine/Plugins/Tests/InterchangeTests/Source/InterchangeTests/Public/ImportTestFunctions/ActorImportTestFunctions.h
@@ -0,0 +1,26 @@
+// Copyright Epic Games, Inc. All Rights Reserved.
+
+#pragma once
+
+#include "ImportTestFunctionsBase.h"
+#include "InterchangeTestFunction.h"
+#include "ActorImportTestFunctions.generated.h"
+
+class AActor;
+struct FInterchangeTestFunctionResult;
+
+
+UCLASS()
+class INTERCHANGETESTS_API UActorImportTestFunctions : public UImportTestFunctionsBase
+{
+	GENERATED_BODY()
+
+public:
+
+	// UImportTestFunctionsBase interface
+	virtual UClass* GetAssociatedAssetType() const override;
+
+	/** Check whether the expected number of actors are imported */
+	UFUNCTION(Exec)
+	static FInterchangeTestFunctionResult CheckImportedActorCount(const TArray<AActor*>& Actors, int32 ExpectedNumberOfImportedActors);
+};
diff --git a/Engine/Plugins/Tests/InterchangeTests/Source/InterchangeTests/Public/InterchangeImportTestStepBase.h b/Engine/Plugins/Tests/InterchangeTests/Source/InterchangeTests/Public/InterchangeImportTestStepBase.h
--- a/Engine/Plugins/Tests/InterchangeTests/Source/InterchangeTests/Public/InterchangeImportTestStepBase.h
+++ b/Engine/Plugins/Tests/InterchangeTests/Source/InterchangeTests/Public/InterchangeImportTestStepBase.h
@@ -30,7 +30,8 @@ public:
 	TArray<FInterchangeTestFunction> Tests;
 
 public:
-	virtual UE::Interchange::FAssetImportResultPtr StartStep(FInterchangeImportTestData& Data) PURE_VIRTUAL(UInterchangeImportTestStepBase::StartStep, return {}; );
+	virtual TTuple<UE::Interchange::FAssetImportResultPtr, UE::Interchange::FSceneImportResultPtr>
+		StartStep(FInterchangeImportTestData& Data) PURE_VIRTUAL(UInterchangeImportTestStepBase::StartStep, return {}; );
 	virtual FTestStepResults FinishStep(FInterchangeImportTestData& Data, FAutomationTestExecutionInfo& ExecutionInfo) PURE_VIRTUAL(UInterchangeImportTestStepBase::FinishStep, return {}; );
 	virtual FString GetContextString() const PURE_VIRTUAL(UInterchangeImportTestStepBase::GetContextString, return {}; );
 
diff --git a/Engine/Plugins/Tests/InterchangeTests/Source/InterchangeTests/Public/InterchangeImportTestStepImport.h b/Engine/Plugins/Tests/InterchangeTests/Source/InterchangeTests/Public/InterchangeImportTestStepImport.h
--- a/Engine/Plugins/Tests/InterchangeTests/Source/InterchangeTests/Public/InterchangeImportTestStepImport.h
+++ b/Engine/Plugins/Tests/InterchangeTests/Source/InterchangeTests/Public/InterchangeImportTestStepImport.h
@@ -29,12 +29,17 @@ public:
 	bool bEmptyDestinationFolderPriorToImport = true;
 
 	/** Whether imported assets should be saved and freshly reloaded after import */
-	UPROPERTY(EditAnywhere, Category = General)
+	UPROPERTY(EditAnywhere, Category = General, Meta=(EditCondition="!bImportIntoLevel"))
 	bool bSaveThenReloadImportedAssets = true;
 
+	/** Whether we should use the import into level workflow */
+	UPROPERTY(EditAnywhere, Category = General)
+	bool bImportIntoLevel = false;
+
 public:
 	// UInterchangeImportTestStepBase interface
-	virtual UE::Interchange::FAssetImportResultPtr StartStep(FInterchangeImportTestData& Data) override;
+	virtual TTuple<UE::Interchange::FAssetImportResultPtr, UE::Interchange::FSceneImportResultPtr>
+		StartStep(FInterchangeImportTestData& Data) override;
 	virtual FTestStepResults FinishStep(FInterchangeImportTestData& Data, FAutomationTestExecutionInfo& ExecutionInfo) override;
 	virtual FString GetContextString() const override;
 };
diff --git a/Engine/Plugins/Tests/InterchangeTests/Source/InterchangeTests/Public/InterchangeImportTestStepReimport.h b/Engine/Plugins/Tests/InterchangeTests/Source/InterchangeTests/Public/InterchangeImportTestStepReimport.h
--- a/Engine/Plugins/Tests/InterchangeTests/Source/InterchangeTests/Public/InterchangeImportTestStepReimport.h
+++ b/Engine/Plugins/Tests/InterchangeTests/Source/InterchangeTests/Public/InterchangeImportTestStepReimport.h
@@ -34,7 +34,8 @@ public:
 
 public:
 	// UInterchangeImportTestStepBase interface
-	virtual UE::Interchange::FAssetImportResultPtr StartStep(FInterchangeImportTestData& Data) override;
+	virtual TTuple<UE::Interchange::FAssetImportResultPtr, UE::Interchange::FSceneImportResultPtr>
+		StartStep(FInterchangeImportTestData& Data) override;
 	virtual FTestStepResults FinishStep(FInterchangeImportTestData& Data, FAutomationTestExecutionInfo& ExecutionInfo) override;
 	virtual FString GetContextString() const override;
 };
