Add asset test functions, to allow package metadata queries in the interchange tests.
#rnx

diff --git a/Engine/Plugins/Tests/InterchangeTests/Source/InterchangeTests/Private/ImportTestFunctions/AssetImportTestFunctions.cpp b/Engine/Plugins/Tests/InterchangeTests/Source/InterchangeTests/Private/ImportTestFunctions/AssetImportTestFunctions.cpp
new file mode 100644
--- /dev/null
+++ b/Engine/Plugins/Tests/InterchangeTests/Source/InterchangeTests/Private/ImportTestFunctions/AssetImportTestFunctions.cpp
@@ -0,0 +1,66 @@
+// Copyright Epic Games, Inc. All Rights Reserved.
+
+#include "ImportTestFunctions/AssetImportTestFunctions.h"
+
+#include "UObject/MetaData.h"
+#include "UObject/UObjectGlobals.h"
+#include "UObject/Package.h"
+
+UClass* UAssetImportTestFunctions::GetAssociatedAssetType() const
+{
+	return UObject::StaticClass();
+}
+
+FInterchangeTestFunctionResult UAssetImportTestFunctions::CheckImportedMetadataCount(const UObject* Object, const int32 ExpectedNumberOfMetadataForThisObject)
+{
+	FInterchangeTestFunctionResult Result;
+	int32 ObjectMetadataCount = 0;
+	if (const TMap<FName, FString>* ObjectMetaDatas = UMetaData::GetMapForObject(Object))
+	{
+		ObjectMetadataCount = ObjectMetaDatas->Num();
+	}
+	if (ObjectMetadataCount != ExpectedNumberOfMetadataForThisObject)
+	{
+		Result.AddError(FString::Printf(TEXT("Expected %d object metadatas, imported %d."), ExpectedNumberOfMetadataForThisObject, ObjectMetadataCount));
+	}
+	return Result;
+}
+
+FInterchangeTestFunctionResult UAssetImportTestFunctions::CheckMetadataExist(const UObject* Object, const FString& ExpectedMetadataKey)
+{
+	FInterchangeTestFunctionResult Result;
+	bool bMetadatExist = false;
+	if (const TMap<FName, FString>* ObjectMetaDatas = UMetaData::GetMapForObject(Object))
+	{
+		bMetadatExist = ObjectMetaDatas->Contains(FName(*ExpectedMetadataKey));
+	}
+	if (!bMetadatExist)
+	{
+		Result.AddError(FString::Printf(TEXT("Expected object metadata key %s was not imported."), *ExpectedMetadataKey));
+	}
+	return Result;
+}
+
+FInterchangeTestFunctionResult UAssetImportTestFunctions::CheckMetadataValue(const UObject* Object, const FString& ExpectedMetadataKey, const FString& ExpectedMetadataValue)
+{
+	FInterchangeTestFunctionResult Result;
+	bool bKeyExist = false;
+	if (const TMap<FName, FString>* ObjectMetaDatas = UMetaData::GetMapForObject(Object))
+	{
+		if (const FString* MetadataValue = ObjectMetaDatas->Find(FName(*ExpectedMetadataKey)))
+		{
+			bKeyExist = true;
+			if (!ExpectedMetadataValue.Equals(*MetadataValue))
+			{
+				Result.AddError(FString::Printf(TEXT("Expected object metadata key [%s] value [%s], found a different value [%s]."), *ExpectedMetadataKey, *ExpectedMetadataValue, *(*MetadataValue)));
+			}
+		}
+	}
+	
+	if(!bKeyExist)
+	{
+		Result.AddError(FString::Printf(TEXT("Expected object metadata key [%s] value [%s], the key was not imported."), *ExpectedMetadataKey, *ExpectedMetadataValue));
+	}
+	return Result;
+}
+
diff --git a/Engine/Plugins/Tests/InterchangeTests/Source/InterchangeTests/Public/ImportTestFunctions/AssetImportTestFunctions.h b/Engine/Plugins/Tests/InterchangeTests/Source/InterchangeTests/Public/ImportTestFunctions/AssetImportTestFunctions.h
new file mode 100644
--- /dev/null
+++ b/Engine/Plugins/Tests/InterchangeTests/Source/InterchangeTests/Public/ImportTestFunctions/AssetImportTestFunctions.h
@@ -0,0 +1,43 @@
+// Copyright Epic Games, Inc. All Rights Reserved.
+
+#pragma once
+
+#include "ImportTestFunctionsBase.h"
+#include "InterchangeTestFunction.h"
+#include "AssetImportTestFunctions.generated.h"
+
+struct FInterchangeTestFunctionResult;
+class UObject;
+
+
+UCLASS()
+class INTERCHANGETESTS_API UAssetImportTestFunctions : public UImportTestFunctionsBase
+{
+	GENERATED_BODY()
+
+public:
+
+	// UImportTestFunctionsBase interface
+	virtual UClass* GetAssociatedAssetType() const override;
+
+
+	/** Check whether the expected number of metadata for the object are imported */
+	UFUNCTION(Exec)
+	static FInterchangeTestFunctionResult CheckImportedMetadataCount(const UObject* Object, const int32 ExpectedNumberOfMetadataForThisObject);
+
+	/** 
+	 * Check whether the expected object metadata key exist.
+	 * @Param ExpectedMetadataKey - The object metadata key to pass to the package to retrieve the metadata value
+	 */
+	UFUNCTION(Exec)
+	static FInterchangeTestFunctionResult CheckMetadataExist(const UObject* Object, const FString& ExpectedMetadataKey);
+
+	/**
+	 * Check whether the expected object metadata value is imported.
+	 * @Param ExpectedMetadataKey - The object metadata key to pass to the package to retrieve the metadata value
+	 * @Param ExpectedMetadataValue - The value to compare the object metadata query with the metadata key
+	 */
+	UFUNCTION(Exec)
+	static FInterchangeTestFunctionResult CheckMetadataValue(const UObject* Object, const FString& ExpectedMetadataKey, const FString& ExpectedMetadataValue);
+
+};
