Allow UObject interchange tests function to work properly
#rnx

diff --git a/Engine/Plugins/Tests/InterchangeTests/Source/InterchangeTests/Private/InterchangeTestFunction.cpp b/Engine/Plugins/Tests/InterchangeTests/Source/InterchangeTests/Private/InterchangeTestFunction.cpp
--- a/Engine/Plugins/Tests/InterchangeTests/Source/InterchangeTests/Private/InterchangeTestFunction.cpp
+++ b/Engine/Plugins/Tests/InterchangeTests/Source/InterchangeTests/Private/InterchangeTestFunction.cpp
@@ -214,8 +214,7 @@ FInterchangeTestFunctionResult FInterchangeTestFunction::Invoke(const TArray<UOb
 	TArray<UObject*> AssetsToPass;
 	AssetsToPass.Reserve(AssetsToTest.Num());
 
-	UClass* ExpectedClass = GetExpectedClass();
-
+	UClass* ExpectedClass = AssetClass;
 	for (UObject* Asset : AssetsToTest)
 	{
 		if (Asset->GetClass()->IsChildOf(ExpectedClass))
