Fix crash when inspecting "skinned vertex count for bones".

#rnx

diff --git a/Engine/Plugins/Tests/InterchangeTests/Source/InterchangeTests/Private/ImportTestFunctions/SkeletalMeshImportTestFunctions.cpp b/Engine/Plugins/Tests/InterchangeTests/Source/InterchangeTests/Private/ImportTestFunctions/SkeletalMeshImportTestFunctions.cpp
--- a/Engine/Plugins/Tests/InterchangeTests/Source/InterchangeTests/Private/ImportTestFunctions/SkeletalMeshImportTestFunctions.cpp
+++ b/Engine/Plugins/Tests/InterchangeTests/Source/InterchangeTests/Private/ImportTestFunctions/SkeletalMeshImportTestFunctions.cpp
@@ -353,7 +353,7 @@ FInterchangeTestFunctionResult USkeletalMeshImportTestFunctions::CheckSkinnedVer
 
 			if (SkinnedVerticesForBone != ExpectedSkinnedVertexCount)
 			{
-				Result.AddError(FString::Printf(TEXT("For bone '%s', expected %d vertices, imported %s."), *BoneName, ExpectedSkinnedVertexCount, SkinnedVerticesForBone));
+				Result.AddError(FString::Printf(TEXT("For bone '%s', expected %d vertices, imported %d."), *BoneName, ExpectedSkinnedVertexCount, SkinnedVerticesForBone));
 			}
 		}
 		else
