Interchange: Downgrade the disable test warning to info level to avoid CIS break when we disable a test. Add a configurable string we can set to explain the reason the test is disabled.

#rnx

diff --git a/Engine/Plugins/Tests/InterchangeTests/Source/InterchangeTests/Private/InterchangeImportTest.cpp b/Engine/Plugins/Tests/InterchangeTests/Source/InterchangeTests/Private/InterchangeImportTest.cpp
--- a/Engine/Plugins/Tests/InterchangeTests/Source/InterchangeTests/Private/InterchangeImportTest.cpp
+++ b/Engine/Plugins/Tests/InterchangeTests/Source/InterchangeTests/Private/InterchangeImportTest.cpp
@@ -108,7 +108,18 @@ bool FInterchangeImportTest::RunTest(const FString& Path)
 				}
 				else
 				{
-					ExecutionInfo.AddWarning(FString::Printf(TEXT("InterchangeImportTestPlan %s is disabled."), *AssetObject->GetName()));
+					FString InfoMessage;
+					if (!AssetObject->DisabledTestReason.IsEmpty())
+					{
+						InfoMessage = FString::Printf(TEXT("InterchangeImportTestPlan %s is disabled because [%s]."), *AssetObject->GetName(), *AssetObject->DisabledTestReason);
+						
+					}
+					else
+					{
+						InfoMessage = FString::Printf(TEXT("InterchangeImportTestPlan %s is disabled."), *AssetObject->GetName());
+					}
+					//Add the disable test as an informative event
+					ExecutionInfo.AddEvent(FAutomationEvent(EAutomationEventType::Info, InfoMessage));
 				}
 
 				return true;
diff --git a/Engine/Plugins/Tests/InterchangeTests/Source/InterchangeTests/Public/InterchangeImportTestPlan.h b/Engine/Plugins/Tests/InterchangeTests/Source/InterchangeTests/Public/InterchangeImportTestPlan.h
--- a/Engine/Plugins/Tests/InterchangeTests/Source/InterchangeTests/Public/InterchangeImportTestPlan.h
+++ b/Engine/Plugins/Tests/InterchangeTests/Source/InterchangeTests/Public/InterchangeImportTestPlan.h
@@ -35,6 +35,10 @@ public:
 	UPROPERTY(EditAnywhere, Category = Activation)
 	bool bIsEnabledInAutomationTests = true;
 
+	/** Why is the test plan disabled. */
+	UPROPERTY(EditAnywhere, Category = Activation, Meta = (EditCondition = "!bIsEnabledInAutomationTests"))
+	FString DisabledTestReason;
+
 	/** Click here to immediately run this single test through the automation framework */
 	UFUNCTION(BlueprintCallable, CallInEditor, Category = Run, Meta = (DisplayPriority = 1))
 	void RunThisTest();
