Interchange Tests: Added a function to validate the number of imported textures.

diff --git a/Engine/Plugins/Tests/InterchangeTests/Source/InterchangeTests/Private/ImportTestFunctions/TextureImportTestFunctions.cpp b/Engine/Plugins/Tests/InterchangeTests/Source/InterchangeTests/Private/ImportTestFunctions/TextureImportTestFunctions.cpp
--- a/Engine/Plugins/Tests/InterchangeTests/Source/InterchangeTests/Private/ImportTestFunctions/TextureImportTestFunctions.cpp
+++ b/Engine/Plugins/Tests/InterchangeTests/Source/InterchangeTests/Private/ImportTestFunctions/TextureImportTestFunctions.cpp
@@ -8,3 +8,14 @@ UClass* UTextureImportTestFunctions::GetAssociatedAssetType() const
 {
 	return UTexture::StaticClass();
 }
+
+FInterchangeTestFunctionResult UTextureImportTestFunctions::CheckImportedTextureCount(const TArray<UTexture*>& Textures, int32 ExpectedNumberOfImportedTextures)
+{
+	FInterchangeTestFunctionResult Result;
+	if (Textures.Num() != ExpectedNumberOfImportedTextures)
+	{
+		Result.AddError(FString::Printf(TEXT("Expected %d textures, imported %d."), ExpectedNumberOfImportedTextures, Textures.Num()));
+	}
+
+	return Result;
+}
\ No newline at end of file
diff --git a/Engine/Plugins/Tests/InterchangeTests/Source/InterchangeTests/Public/ImportTestFunctions/TextureImportTestFunctions.h b/Engine/Plugins/Tests/InterchangeTests/Source/InterchangeTests/Public/ImportTestFunctions/TextureImportTestFunctions.h
--- a/Engine/Plugins/Tests/InterchangeTests/Source/InterchangeTests/Public/ImportTestFunctions/TextureImportTestFunctions.h
+++ b/Engine/Plugins/Tests/InterchangeTests/Source/InterchangeTests/Public/ImportTestFunctions/TextureImportTestFunctions.h
@@ -18,4 +18,8 @@ public:
 
 	// UImportTestFunctionsBase interface
 	virtual UClass* GetAssociatedAssetType() const override;
+
+	/** Check whether the expected number of textures are imported */
+	UFUNCTION(Exec)
+	static FInterchangeTestFunctionResult CheckImportedTextureCount(const TArray<UTexture*>& Textures, int32 ExpectedNumberOfImportedTextures);
 };
