Interchange Tests
- Changed CheckImportedMaterialCount to count only UMaterials.
- Added CheckImportedMaterialInstanceCount to count only UMaterialInstances.

diff --git a/Engine/Plugins/Tests/InterchangeTests/Source/InterchangeTests/Private/ImportTestFunctions/MaterialImportTestFunctions.cpp b/Engine/Plugins/Tests/InterchangeTests/Source/InterchangeTests/Private/ImportTestFunctions/MaterialImportTestFunctions.cpp
--- a/Engine/Plugins/Tests/InterchangeTests/Source/InterchangeTests/Private/ImportTestFunctions/MaterialImportTestFunctions.cpp
+++ b/Engine/Plugins/Tests/InterchangeTests/Source/InterchangeTests/Private/ImportTestFunctions/MaterialImportTestFunctions.cpp
@@ -1,22 +1,59 @@
 // Copyright Epic Games, Inc. All Rights Reserved.
 
 #include "ImportTestFunctions/MaterialImportTestFunctions.h"
+
+#include "Materials/Material.h"
+#include "Materials/MaterialInstance.h"
 #include "Materials/MaterialInterface.h"
 
+namespace UE::Interchange::Tests::Private
+{
+	TArray<UMaterialInterface*> FilterByMaterialInterfaceSubclass(const TArray<UMaterialInterface*>& MaterialInterfaces, const TSubclassOf<UMaterialInterface>& MaterialInterfaceSubclass)
+	{
+		TArray<UMaterialInterface*> FilteredMaterialInterfaces;
+
+		for (UMaterialInterface* MaterialInterface : MaterialInterfaces)
+		{
+			if (MaterialInterface->IsA(MaterialInterfaceSubclass))
+			{
+				FilteredMaterialInterfaces.Add(MaterialInterface);
+			}
+		}
+
+		return FilteredMaterialInterfaces;
+	}
+}
 
 UClass* UMaterialImportTestFunctions::GetAssociatedAssetType() const
 {
 	return UMaterialInterface::StaticClass();
 }
 
+FInterchangeTestFunctionResult UMaterialImportTestFunctions::CheckImportedMaterialCount(const TArray<UMaterialInterface*>& MaterialInterfaces, int32 ExpectedNumberOfImportedMaterials)
+{
+	using namespace UE::Interchange::Tests::Private;
 
+	FInterchangeTestFunctionResult Result;
+	const int32 NumImportedMaterials = FilterByMaterialInterfaceSubclass(MaterialInterfaces, UMaterial::StaticClass()).Num();
 
-FInterchangeTestFunctionResult UMaterialImportTestFunctions::CheckImportedMaterialCount(const TArray<UMaterialInterface*>& Materials, int32 ExpectedNumberOfImportedMaterials)
+	if (NumImportedMaterials != ExpectedNumberOfImportedMaterials)
+	{
+		Result.AddError(FString::Printf(TEXT("Expected %d materials, imported %d."), ExpectedNumberOfImportedMaterials, NumImportedMaterials));
+	}
+
+	return Result;
+}
+
+FInterchangeTestFunctionResult UMaterialImportTestFunctions::CheckImportedMaterialInstanceCount(const TArray<UMaterialInterface*>& MaterialInterfaces, int32 ExpectedNumberOfImportedMaterialInstances)
 {
+	using namespace UE::Interchange::Tests::Private;
+
 	FInterchangeTestFunctionResult Result;
-	if (Materials.Num() != ExpectedNumberOfImportedMaterials)
+	const int32 NumImportedMaterialInstances = FilterByMaterialInterfaceSubclass(MaterialInterfaces, UMaterialInstance::StaticClass()).Num();
+
+	if (NumImportedMaterialInstances != ExpectedNumberOfImportedMaterialInstances)
 	{
-		Result.AddError(FString::Printf(TEXT("Expected %d materials, imported %d."), ExpectedNumberOfImportedMaterials, Materials.Num()));
+		Result.AddError(FString::Printf(TEXT("Expected %d material instances, imported %d."), ExpectedNumberOfImportedMaterialInstances, NumImportedMaterialInstances));
 	}
 
 	return Result;
diff --git a/Engine/Plugins/Tests/InterchangeTests/Source/InterchangeTests/Public/ImportTestFunctions/MaterialImportTestFunctions.h b/Engine/Plugins/Tests/InterchangeTests/Source/InterchangeTests/Public/ImportTestFunctions/MaterialImportTestFunctions.h
--- a/Engine/Plugins/Tests/InterchangeTests/Source/InterchangeTests/Public/ImportTestFunctions/MaterialImportTestFunctions.h
+++ b/Engine/Plugins/Tests/InterchangeTests/Source/InterchangeTests/Public/ImportTestFunctions/MaterialImportTestFunctions.h
@@ -21,5 +21,9 @@ public:
 
 	/** Check whether the expected number of materials are imported */
 	UFUNCTION(Exec)
-	static FInterchangeTestFunctionResult CheckImportedMaterialCount(const TArray<UMaterialInterface*>& Materials, int32 ExpectedNumberOfImportedMaterials);
+	static FInterchangeTestFunctionResult CheckImportedMaterialCount(const TArray<UMaterialInterface*>& MaterialInterfaces, int32 ExpectedNumberOfImportedMaterials);
+
+	/** Check whether the expected number of material instances are imported */
+	UFUNCTION(Exec)
+	static FInterchangeTestFunctionResult CheckImportedMaterialInstanceCount(const TArray<UMaterialInterface*>& MaterialInterfaces, int32 ExpectedNumberOfImportedMaterialInstances);
 };
