Interchange Tests
- Added a check against a static mesh source model build settings.
- Added a check against a static mesh nanite settings.

diff --git a/Engine/Plugins/Tests/InterchangeTests/Source/InterchangeTests/Private/ImportTestFunctions/StaticMeshImportTestFunctions.cpp b/Engine/Plugins/Tests/InterchangeTests/Source/InterchangeTests/Private/ImportTestFunctions/StaticMeshImportTestFunctions.cpp
--- a/Engine/Plugins/Tests/InterchangeTests/Source/InterchangeTests/Private/ImportTestFunctions/StaticMeshImportTestFunctions.cpp
+++ b/Engine/Plugins/Tests/InterchangeTests/Source/InterchangeTests/Private/ImportTestFunctions/StaticMeshImportTestFunctions.cpp
@@ -8,6 +8,46 @@
 #include "StaticMeshAttributes.h"
 #include "ObjectTools.h"
 
+namespace UE::Interchange::Tests::Private
+{
+	void ReportPropertyDeltaAsErrors(FInterchangeTestFunctionResult& Result, TPropertyValueIterator<const FProperty>& PropertyValueIteratorA, TPropertyValueIterator<const FProperty>& PropertyValueIteratorB)
+	{
+		for (; PropertyValueIteratorA && PropertyValueIteratorB; ++PropertyValueIteratorA, ++PropertyValueIteratorB)
+		{
+			const FProperty* PropertyA = PropertyValueIteratorA->Key;
+			const FProperty* PropertyB = PropertyValueIteratorB->Key;
+
+			if (!PropertyA || !PropertyB)
+			{
+				Result.AddError(TEXT("Unexpected error (invalid FProperty)"));
+			}
+
+			//if (!PropertyA->SameType(PropertyB))
+			{
+				Result.AddError(FString::Printf(TEXT("Unexpected error (FProperty %s::%s doesn't match FProperty %s::%s)"),
+					*PropertyA->GetOwnerStruct()->GetName(), *PropertyA->GetName(),
+					*PropertyB->GetOwnerStruct()->GetName(), *PropertyB->GetName()));
+			}
+
+			const void* ValueA = PropertyValueIteratorA->Value;
+			const void* ValueB = PropertyValueIteratorB->Value;
+
+			if (!PropertyA->Identical(ValueA, ValueB))
+			{
+				FString ValueAString;
+				PropertyA->ExportTextItem_Direct(ValueAString, ValueA, nullptr, nullptr, PPF_None);
+
+				FString ValueBString;
+				PropertyB->ExportTextItem_Direct(ValueBString, ValueB, nullptr, nullptr, PPF_None);
+
+				Result.AddError(FString::Printf(TEXT("Expected property (%s::%s=%s), imported (%s::%s=%s)"),
+					*PropertyA->GetOwnerStruct()->GetName(), *PropertyA->GetName(), *ValueAString,
+					*PropertyB->GetOwnerStruct()->GetName(), *PropertyB->GetName(), *ValueBString));
+			}
+		}
+	}
+}
+
 
 UClass* UStaticMeshImportTestFunctions::GetAssociatedAssetType() const
 {
@@ -710,3 +750,39 @@ FInterchangeTestFunctionResult UStaticMeshImportTestFunctions::CheckAgainstGroun
 	return Result;
 }
 
+FInterchangeTestFunctionResult UStaticMeshImportTestFunctions::CheckBuildSettings(UStaticMesh* Mesh, int32 LodIndex, const FMeshBuildSettings& ExpectedBuildSettings)
+{
+	using namespace UE::Interchange::Tests::Private;
+
+	FInterchangeTestFunctionResult Result;
+
+	if (Mesh->GetSourceModel(LodIndex).BuildSettings != ExpectedBuildSettings)
+	{
+		const UStruct* MeshBuildSettingsStruct = FMeshBuildSettings::StaticStruct();
+		TPropertyValueIterator<const FProperty> MeshPropertyValueIterator(MeshBuildSettingsStruct, &Mesh->GetSourceModel(LodIndex).BuildSettings);
+		TPropertyValueIterator<const FProperty> ExpectedPropertyValueIterator(MeshBuildSettingsStruct, &ExpectedBuildSettings);
+
+		ReportPropertyDeltaAsErrors(Result, MeshPropertyValueIterator, ExpectedPropertyValueIterator);
+
+	}
+
+	return Result;
+}
+
+FInterchangeTestFunctionResult UStaticMeshImportTestFunctions::CheckNaniteSettings(UStaticMesh* Mesh, const FMeshNaniteSettings& ExpectedNaniteSettings)
+{
+	using namespace UE::Interchange::Tests::Private;
+
+	FInterchangeTestFunctionResult Result;
+
+	if (Mesh->NaniteSettings != ExpectedNaniteSettings)
+	{
+		const UStruct* NaniteSettingsStruct = FMeshNaniteSettings::StaticStruct();
+		TPropertyValueIterator<const FProperty> MeshPropertyValueIterator(NaniteSettingsStruct, &Mesh->NaniteSettings);
+		TPropertyValueIterator<const FProperty> ExpectedPropertyValueIterator(NaniteSettingsStruct, &ExpectedNaniteSettings);
+
+		ReportPropertyDeltaAsErrors(Result, MeshPropertyValueIterator, ExpectedPropertyValueIterator);
+	}
+
+	return Result;
+}
diff --git a/Engine/Plugins/Tests/InterchangeTests/Source/InterchangeTests/Public/ImportTestFunctions/StaticMeshImportTestFunctions.h b/Engine/Plugins/Tests/InterchangeTests/Source/InterchangeTests/Public/ImportTestFunctions/StaticMeshImportTestFunctions.h
--- a/Engine/Plugins/Tests/InterchangeTests/Source/InterchangeTests/Public/ImportTestFunctions/StaticMeshImportTestFunctions.h
+++ b/Engine/Plugins/Tests/InterchangeTests/Source/InterchangeTests/Public/ImportTestFunctions/StaticMeshImportTestFunctions.h
@@ -128,4 +128,12 @@ public:
 		bool bCheckVertexPositionsEqual = true,
 		bool bCheckNormalsEqual = true
 	);
+
+	/** Check whether the mesh source model has the expected build settings */
+	UFUNCTION(Exec)
+	static FInterchangeTestFunctionResult CheckBuildSettings(UStaticMesh* Mesh, int32 LodIndex, const FMeshBuildSettings& ExpectedBuildSettings);
+
+	/** Check whether the mesh has the expected Nanite settings */
+	UFUNCTION(Exec)
+	static FInterchangeTestFunctionResult CheckNaniteSettings(UStaticMesh* Mesh, const FMeshNaniteSettings& ExpectedNaniteSettings);
 };
