Fixed Interchange static mesh build & nanite settings always logging unexpected errors.

diff --git a/Engine/Plugins/Tests/InterchangeTests/Source/InterchangeTests/Private/ImportTestFunctions/StaticMeshImportTestFunctions.cpp b/Engine/Plugins/Tests/InterchangeTests/Source/InterchangeTests/Private/ImportTestFunctions/StaticMeshImportTestFunctions.cpp
--- a/Engine/Plugins/Tests/InterchangeTests/Source/InterchangeTests/Private/ImportTestFunctions/StaticMeshImportTestFunctions.cpp
+++ b/Engine/Plugins/Tests/InterchangeTests/Source/InterchangeTests/Private/ImportTestFunctions/StaticMeshImportTestFunctions.cpp
@@ -22,7 +22,7 @@ namespace UE::Interchange::Tests::Private
 				Result.AddError(TEXT("Unexpected error (invalid FProperty)"));
 			}
 
-			//if (!PropertyA->SameType(PropertyB))
+			if (!PropertyA->SameType(PropertyB))
 			{
 				Result.AddError(FString::Printf(TEXT("Unexpected error (FProperty %s::%s doesn't match FProperty %s::%s)"),
 					*PropertyA->GetOwnerStruct()->GetName(), *PropertyA->GetName(),
