Static analysis fix.

diff --git a/Engine/Plugins/Tests/InterchangeTests/Source/InterchangeTests/Private/ImportTestFunctions/StaticMeshImportTestFunctions.cpp b/Engine/Plugins/Tests/InterchangeTests/Source/InterchangeTests/Private/ImportTestFunctions/StaticMeshImportTestFunctions.cpp
--- a/Engine/Plugins/Tests/InterchangeTests/Source/InterchangeTests/Private/ImportTestFunctions/StaticMeshImportTestFunctions.cpp
+++ b/Engine/Plugins/Tests/InterchangeTests/Source/InterchangeTests/Private/ImportTestFunctions/StaticMeshImportTestFunctions.cpp
@@ -20,6 +20,7 @@ namespace UE::Interchange::Tests::Private
 			if (!PropertyA || !PropertyB)
 			{
 				Result.AddError(TEXT("Unexpected error (invalid FProperty)"));
+				return;
 			}
 
 			if (!PropertyA->SameType(PropertyB))
