Added test functions specific to MaterialX import

diff --git a/Engine/Plugins/Tests/InterchangeTests/Source/InterchangeTests/Private/ImportTestFunctions/MaterialXTestFunctions.cpp b/Engine/Plugins/Tests/InterchangeTests/Source/InterchangeTests/Private/ImportTestFunctions/MaterialXTestFunctions.cpp
new file mode 100644
--- /dev/null
+++ b/Engine/Plugins/Tests/InterchangeTests/Source/InterchangeTests/Private/ImportTestFunctions/MaterialXTestFunctions.cpp
@@ -0,0 +1,98 @@
+// Copyright Epic Games, Inc. All Rights Reserved.
+
+#include "ImportTestFunctions/MaterialXTestFunctions.h"
+
+#include "InterchangeTestFunction.h"
+
+#include "Materials/Material.h"
+#include "Materials/MaterialExpressionMaterialFunctionCall.h"
+#include "Materials/MaterialInterface.h"
+#include "Materials/MaterialInterface.h"
+
+#pragma optimize("",off)
+namespace UE::Interchange::Tests::Private
+{
+	void GetMaterialFunctionConnectedInputs(const UMaterial& Material, TArray<FString>& InputNames)
+	{
+		if (!ensure(Material.HasBaseColorConnected()))
+		{
+			return;
+		}
+
+#if WITH_EDITORONLY_DATA
+		const UMaterialEditorOnlyData* MaterialEditorOnly = Material.GetEditorOnlyData();
+		if (UMaterialExpressionMaterialFunctionCall* FunctionCall = Cast<UMaterialExpressionMaterialFunctionCall>(MaterialEditorOnly->BaseColor.Expression))
+		{
+			const TArray<FExpressionInput*>& Inputs = FunctionCall->GetInputs();
+			InputNames.Empty(Inputs.Num());
+			for (const FExpressionInput* InputExpression : Inputs)
+			{
+				if (InputExpression && InputExpression->Expression)
+				{
+					InputNames.Add(InputExpression->InputName.ToString());
+				}
+			}
+		}
+#endif
+	}
+}
+
+UClass* UMaterialXTestFunctions::GetAssociatedAssetType() const
+{
+	return UMaterialInterface::StaticClass();
+}
+
+FInterchangeTestFunctionResult UMaterialXTestFunctions::CheckConnectedInputCount(const UMaterialInterface* MaterialInterface, int32 ExpectedNumber)
+{
+	using namespace UE::Interchange::Tests::Private;
+
+	FInterchangeTestFunctionResult Result;
+
+	if (const UMaterial* Material = Cast<UMaterial>(MaterialInterface))
+	{
+		TArray<FString> InputNames;
+		GetMaterialFunctionConnectedInputs(*Material, InputNames);
+		if (ExpectedNumber != InputNames.Num())
+		{
+			Result.AddError(FString::Printf(TEXT("Expected %d connected inputs, found %d connected."), ExpectedNumber, InputNames.Num()));
+		}
+	}
+	else
+	{
+		Result.AddError(FString::Printf(TEXT("Expected a %s got a %s"), *UMaterial::StaticClass()->GetName(), *MaterialInterface->GetClass()->GetName()));
+	}
+
+	return Result;
+}
+
+FInterchangeTestFunctionResult UMaterialXTestFunctions::CheckInputConnected(const UMaterialInterface* MaterialInterface, const FString& InputName, bool bIsConnected)
+{
+	using namespace UE::Interchange::Tests::Private;
+
+	FInterchangeTestFunctionResult Result;
+
+	if (const UMaterial* Material = Cast<UMaterial>(MaterialInterface))
+	{
+		TArray<FString> InputNames;
+		GetMaterialFunctionConnectedInputs(*Material, InputNames);
+
+		if (InputNames.Contains(InputName) != bIsConnected)
+		{
+			if (bIsConnected)
+			{
+				Result.AddError(FString::Printf(TEXT("Expected input %s to be connected, found not connected."), *InputName));
+			}
+			else
+			{
+				Result.AddError(FString::Printf(TEXT("Expected input %s not to be connected, found connected."), *InputName));
+			}
+		}
+	}
+	else
+	{
+		Result.AddError(FString::Printf(TEXT("Expected a %s got a %s"), *UMaterial::StaticClass()->GetName(), *MaterialInterface->GetClass()->GetName()));
+	}
+
+	return Result;
+}
+#pragma optimize("",on)
diff --git a/Engine/Plugins/Tests/InterchangeTests/Source/InterchangeTests/Public/ImportTestFunctions/MaterialXTestFunctions.h b/Engine/Plugins/Tests/InterchangeTests/Source/InterchangeTests/Public/ImportTestFunctions/MaterialXTestFunctions.h
new file mode 100644
--- /dev/null
+++ b/Engine/Plugins/Tests/InterchangeTests/Source/InterchangeTests/Public/ImportTestFunctions/MaterialXTestFunctions.h
@@ -0,0 +1,30 @@
+// Copyright Epic Games, Inc. All Rights Reserved.
+
+#pragma once
+
+#include "ImportTestFunctionsBase.h"
+
+#include "MaterialXTestFunctions.generated.h"
+
+class UMaterialInterface;
+struct FInterchangeTestFunctionResult;
+
+
+UCLASS()
+class INTERCHANGETESTS_API UMaterialXTestFunctions : public UImportTestFunctionsBase
+{
+	GENERATED_BODY()
+
+public:
+
+	// UImportTestFunctionsBase interface
+	virtual UClass* GetAssociatedAssetType() const override;
+
+	/** Check whether the expected number of inputs are connected to the MX_StandardSurface material function */
+	UFUNCTION(Exec, meta = (DisplayName = "MX: Check Connected Input Count"))
+	static FInterchangeTestFunctionResult CheckConnectedInputCount(const UMaterialInterface* MaterialInterface, int32 ExpectedNumber);
+
+	/** Check whether a specific input of the MX_StandardSurface material function is connected or not */
+	UFUNCTION(Exec, meta = (DisplayName = "MX: Check Input Is Connected"))
+	static FInterchangeTestFunctionResult CheckInputConnected(const UMaterialInterface* MaterialInterface, const FString& InputName, bool bIsConnected);
+};
