Fixed Horde's issue 192302

diff --git a/Engine/Plugins/Tests/InterchangeTests/Source/InterchangeTests/Private/ImportTestFunctions/MaterialXTestFunctions.cpp b/Engine/Plugins/Tests/InterchangeTests/Source/InterchangeTests/Private/ImportTestFunctions/MaterialXTestFunctions.cpp
--- a/Engine/Plugins/Tests/InterchangeTests/Source/InterchangeTests/Private/ImportTestFunctions/MaterialXTestFunctions.cpp
+++ b/Engine/Plugins/Tests/InterchangeTests/Source/InterchangeTests/Private/ImportTestFunctions/MaterialXTestFunctions.cpp
@@ -9,7 +9,6 @@
 #include "Materials/MaterialInterface.h"
 #include "Materials/MaterialInterface.h"
 
-#pragma optimize("",off)
 namespace UE::Interchange::Tests::Private
 {
 	void GetMaterialFunctionConnectedInputs(const UMaterial& Material, TArray<FString>& InputNames)
@@ -95,4 +94,3 @@ FInterchangeTestFunctionResult UMaterialXTestFunctions::CheckInputConnected(cons
 
 	return Result;
 }
-#pragma optimize("",on)
