Fixed Horde's issue 192308

diff --git a/Engine/Plugins/Tests/InterchangeTests/Source/InterchangeTests/Public/ImportTestFunctions/MaterialXTestFunctions.h b/Engine/Plugins/Tests/InterchangeTests/Source/InterchangeTests/Public/ImportTestFunctions/MaterialXTestFunctions.h
--- a/Engine/Plugins/Tests/InterchangeTests/Source/InterchangeTests/Public/ImportTestFunctions/MaterialXTestFunctions.h
+++ b/Engine/Plugins/Tests/InterchangeTests/Source/InterchangeTests/Public/ImportTestFunctions/MaterialXTestFunctions.h
@@ -3,12 +3,11 @@
 #pragma once
 
 #include "ImportTestFunctionsBase.h"
+#include "InterchangeTestFunction.h"
 
 #include "MaterialXTestFunctions.generated.h"
 
 class UMaterialInterface;
-struct FInterchangeTestFunctionResult;
-
 
 UCLASS()
 class INTERCHANGETESTS_API UMaterialXTestFunctions : public UImportTestFunctionsBase
