Interchange: Write some interchange automation test check functions. Modify the automation import tests for skeletalmesh, staticmesh and animation to use the new check function.

#rnx

diff --git a/Engine/Plugins/Tests/InterchangeTests/Source/InterchangeTests/Private/ImportTestFunctions/AnimationImportTestFunctions.cpp b/Engine/Plugins/Tests/InterchangeTests/Source/InterchangeTests/Private/ImportTestFunctions/AnimationImportTestFunctions.cpp
--- a/Engine/Plugins/Tests/InterchangeTests/Source/InterchangeTests/Private/ImportTestFunctions/AnimationImportTestFunctions.cpp
+++ b/Engine/Plugins/Tests/InterchangeTests/Source/InterchangeTests/Private/ImportTestFunctions/AnimationImportTestFunctions.cpp
@@ -1,13 +1,170 @@
 // Copyright Epic Games, Inc. All Rights Reserved.
 
 #include "ImportTestFunctions/AnimationImportTestFunctions.h"
-#include "Engine/SkeletalMesh.h"
+#include "Animation/AnimSequence.h"
+#include "Animation/Skeleton.h"
 
 #include UE_INLINE_GENERATED_CPP_BY_NAME(AnimationImportTestFunctions)
 
+namespace UE::Private::AnimationImportTestFunction
+{
+	static bool GetImportedCustomCurveKey(UAnimSequence* AnimSequence, const FString& CurveName, const int32 KeyIndex, FRichCurveKey& OutCurveKey, FInterchangeTestFunctionResult& Result)
+	{
+		const FFloatCurve* FloatCurve = nullptr;
+		FSmartName OutSmartName;
+		if (AnimSequence->GetSkeleton()->GetSmartNameByName(USkeleton::AnimCurveMappingName, *CurveName, OutSmartName))
+		{
+			for (const FFloatCurve& Curve : AnimSequence->GetCurveData().FloatCurves)
+			{
+				if (Curve.Name.DisplayName == FName(*CurveName))
+				{
+					FloatCurve = &Curve;
+					break;
+				}
+			}
+		}
+		if (FloatCurve == nullptr)
+		{
+			Result.AddError(FString::Printf(TEXT("No custom curve named %s was imported"), *CurveName));
+			return false;
+		}
+
+		const FRichCurve& RichCurve = FloatCurve->FloatCurve;
+		const TArray<FRichCurveKey>& Keys = RichCurve.GetConstRefOfKeys();
+		if (!Keys.IsValidIndex(KeyIndex))
+		{
+			Result.AddError(FString::Printf(TEXT("No key at the index %d was imported in curve %s"), KeyIndex, *CurveName));
+			return false;
+		}
+
+		OutCurveKey = Keys[KeyIndex];
+		return true;
+	}
+}
 
 UClass* UAnimationImportTestFunctions::GetAssociatedAssetType() const
 {
-	return USkeletalMesh::StaticClass();
+	return UAnimSequence::StaticClass();
+}
+
+FInterchangeTestFunctionResult UAnimationImportTestFunctions::CheckImportedAnimSequenceCount(const TArray<UAnimSequence*>& AnimSequences, int32 ExpectedNumberOfImportedAnimSequences)
+{
+	FInterchangeTestFunctionResult Result;
+	if (AnimSequences.Num() != ExpectedNumberOfImportedAnimSequences)
+	{
+		Result.AddError(FString::Printf(TEXT("Expected %d skeletal meshes, imported %d."), ExpectedNumberOfImportedAnimSequences, AnimSequences.Num()));
+	}
+
+	return Result;
+}
+
+FInterchangeTestFunctionResult UAnimationImportTestFunctions::CheckAnimationLength(UAnimSequence* AnimSequence, float ExpectedAnimationLength)
+{
+	FInterchangeTestFunctionResult Result;
+
+	float animationLength = AnimSequence->GetPlayLength();
+	if (!FMath::IsNearlyEqual(animationLength, ExpectedAnimationLength, UE_KINDA_SMALL_NUMBER))
+	{
+		Result.AddError(FString::Printf(TEXT("Expected animation length of %f seconds, imported %f."), ExpectedAnimationLength, animationLength));
+	}
+
+	return Result;
+}
+
+FInterchangeTestFunctionResult UAnimationImportTestFunctions::CheckAnimationFrameNumber(UAnimSequence* AnimSequence, int32 ExpectedFrameNumber)
+{
+	FInterchangeTestFunctionResult Result;
+
+	int32 FrameNumber = AnimSequence->GetNumberOfSampledKeys();
+	if (FrameNumber != ExpectedFrameNumber)
+	{
+		Result.AddError(FString::Printf(TEXT("Expected animation frame number %d, imported %d."), ExpectedFrameNumber, FrameNumber));
+	}
+
+	return Result;
+}
+
+FInterchangeTestFunctionResult UAnimationImportTestFunctions::CheckCurveKeyTime(UAnimSequence* AnimSequence, const FString& CurveName, int32 KeyIndex, float ExpectedCurveKeyTime)
+{
+	FInterchangeTestFunctionResult Result;
+	FRichCurveKey CurveKey;
+	if (UE::Private::AnimationImportTestFunction::GetImportedCustomCurveKey(AnimSequence, CurveName, KeyIndex, CurveKey, Result))
+	{
+		if (!FMath::IsNearlyEqual(CurveKey.Time, ExpectedCurveKeyTime, UE_KINDA_SMALL_NUMBER))
+		{
+			Result.AddError(FString::Printf(TEXT("Expected animation curve key time %f, imported %f."), ExpectedCurveKeyTime, CurveKey.Time));
+		}
+	}
+	return Result;
+}
+
+FInterchangeTestFunctionResult UAnimationImportTestFunctions::CheckCurveKeyValue(UAnimSequence* AnimSequence, const FString& CurveName, int32 KeyIndex, float ExpectedCurveKeyValue)
+{
+	FInterchangeTestFunctionResult Result;
+	FRichCurveKey CurveKey;
+	if (UE::Private::AnimationImportTestFunction::GetImportedCustomCurveKey(AnimSequence, CurveName, KeyIndex, CurveKey, Result))
+	{
+		if (!FMath::IsNearlyEqual(CurveKey.Value, ExpectedCurveKeyValue, UE_KINDA_SMALL_NUMBER))
+		{
+			Result.AddError(FString::Printf(TEXT("Expected animation curve key value %f, imported %f."), ExpectedCurveKeyValue, CurveKey.Value));
+		}
+	}
+	return Result;
+}
+
+FInterchangeTestFunctionResult UAnimationImportTestFunctions::CheckCurveKeyArriveTangent(UAnimSequence* AnimSequence, const FString& CurveName, int32 KeyIndex, float ExpectedCurveKeyArriveTangent)
+{
+	FInterchangeTestFunctionResult Result;
+	FRichCurveKey CurveKey;
+	if (UE::Private::AnimationImportTestFunction::GetImportedCustomCurveKey(AnimSequence, CurveName, KeyIndex, CurveKey, Result))
+	{
+		if (!FMath::IsNearlyEqual(CurveKey.ArriveTangent, ExpectedCurveKeyArriveTangent, UE_KINDA_SMALL_NUMBER))
+		{
+			Result.AddError(FString::Printf(TEXT("Expected animation curve key arrive tangent %f, imported %f."), ExpectedCurveKeyArriveTangent, CurveKey.ArriveTangent));
+		}
+	}
+	return Result;
+}
+
+FInterchangeTestFunctionResult UAnimationImportTestFunctions::CheckCurveKeyArriveTangentWeight(UAnimSequence* AnimSequence, const FString& CurveName, int32 KeyIndex, float ExpectedCurveKeyArriveTangentWeight)
+{
+	FInterchangeTestFunctionResult Result;
+	FRichCurveKey CurveKey;
+	if (UE::Private::AnimationImportTestFunction::GetImportedCustomCurveKey(AnimSequence, CurveName, KeyIndex, CurveKey, Result))
+	{
+		if (!FMath::IsNearlyEqual(CurveKey.ArriveTangentWeight, ExpectedCurveKeyArriveTangentWeight, UE_KINDA_SMALL_NUMBER))
+		{
+			Result.AddError(FString::Printf(TEXT("Expected animation curve key arrive tangent weight %f, imported %f."), ExpectedCurveKeyArriveTangentWeight, CurveKey.ArriveTangentWeight));
+		}
+	}
+	return Result;
+}
+
+FInterchangeTestFunctionResult UAnimationImportTestFunctions::CheckCurveKeyLeaveTangent(UAnimSequence* AnimSequence, const FString& CurveName, int32 KeyIndex, float ExpectedCurveKeyLeaveTangent)
+{
+	FInterchangeTestFunctionResult Result;
+	FRichCurveKey CurveKey;
+	if (UE::Private::AnimationImportTestFunction::GetImportedCustomCurveKey(AnimSequence, CurveName, KeyIndex, CurveKey, Result))
+	{
+		if (!FMath::IsNearlyEqual(CurveKey.LeaveTangent, ExpectedCurveKeyLeaveTangent, UE_KINDA_SMALL_NUMBER))
+		{
+			Result.AddError(FString::Printf(TEXT("Expected animation curve key Leave tangent %f, imported %f."), ExpectedCurveKeyLeaveTangent, CurveKey.LeaveTangent));
+		}
+	}
+	return Result;
+}
+
+FInterchangeTestFunctionResult UAnimationImportTestFunctions::CheckCurveKeyLeaveTangentWeight(UAnimSequence* AnimSequence, const FString& CurveName, int32 KeyIndex, float ExpectedCurveKeyLeaveTangentWeight)
+{
+	FInterchangeTestFunctionResult Result;
+	FRichCurveKey CurveKey;
+	if (UE::Private::AnimationImportTestFunction::GetImportedCustomCurveKey(AnimSequence, CurveName, KeyIndex, CurveKey, Result))
+	{
+		if (!FMath::IsNearlyEqual(CurveKey.LeaveTangentWeight, ExpectedCurveKeyLeaveTangentWeight, UE_KINDA_SMALL_NUMBER))
+		{
+			Result.AddError(FString::Printf(TEXT("Expected animation curve key Leave tangent weight %f, imported %f."), ExpectedCurveKeyLeaveTangentWeight, CurveKey.LeaveTangentWeight));
+		}
+	}
+	return Result;
 }
 
diff --git a/Engine/Plugins/Tests/InterchangeTests/Source/InterchangeTests/Private/ImportTestFunctions/SkeletalMeshImportTestFunctions.cpp b/Engine/Plugins/Tests/InterchangeTests/Source/InterchangeTests/Private/ImportTestFunctions/SkeletalMeshImportTestFunctions.cpp
--- a/Engine/Plugins/Tests/InterchangeTests/Source/InterchangeTests/Private/ImportTestFunctions/SkeletalMeshImportTestFunctions.cpp
+++ b/Engine/Plugins/Tests/InterchangeTests/Source/InterchangeTests/Private/ImportTestFunctions/SkeletalMeshImportTestFunctions.cpp
@@ -33,33 +33,72 @@ FInterchangeTestFunctionResult USkeletalMeshImportTestFunctions::CheckRenderVert
 {
 	FInterchangeTestFunctionResult Result;
 
-	// @todo: add separate test function for checking the source data vertex count?
-	int32 ImportedLods = Mesh->GetResourceForRendering()->LODRenderData.Num();
-	if (LodIndex >= ImportedLods)
+	if (Mesh->GetResourceForRendering())
 	{
-		Result.AddError(FString::Printf(TEXT("The imported mesh doesn't contain LOD index %d (imported %d)."), LodIndex, ImportedLods));
+		int32 ImportedLods = Mesh->GetResourceForRendering()->LODRenderData.Num();
+		if (LodIndex < 0 || LodIndex >= ImportedLods)
+		{
+			Result.AddError(FString::Printf(TEXT("The imported mesh doesn't contain LOD index %d (imported %d)."), LodIndex, ImportedLods));
+		}
+		else
+		{
+			int32 RealVertexCount = Mesh->GetResourceForRendering()->LODRenderData[LodIndex].GetNumVertices();
+			if (RealVertexCount != ExpectedNumberOfRenderVertices)
+			{
+				Result.AddError(FString::Printf(TEXT("For LOD %d, expected %d vertices, imported %d."), LodIndex, ExpectedNumberOfRenderVertices, RealVertexCount));
+			}
+		}
 	}
 	else
 	{
-		int32 RealVertexCount = Mesh->GetResourceForRendering()->LODRenderData[LodIndex].GetNumVertices();
-		if (RealVertexCount != ExpectedNumberOfRenderVertices)
+		Result.AddError(FString::Printf(TEXT("No valid render data for this skeletalmesh %s."), *Mesh->GetName()));
+	}
+
+	return Result;
+}
+
+FInterchangeTestFunctionResult USkeletalMeshImportTestFunctions::CheckRenderTriangleCount(USkeletalMesh* Mesh, int32 LodIndex, int32 ExpectedNumberOfRenderTriangles)
+{
+	FInterchangeTestFunctionResult Result;
+
+	if (Mesh->GetResourceForRendering())
+	{
+		int32 ImportedLods = Mesh->GetResourceForRendering()->LODRenderData.Num();
+		if (LodIndex < 0 || LodIndex >= ImportedLods)
 		{
-			Result.AddError(FString::Printf(TEXT("For LOD %d, expected %d vertices, imported %d."), LodIndex, ExpectedNumberOfRenderVertices, RealVertexCount));
+			Result.AddError(FString::Printf(TEXT("The imported mesh doesn't contain LOD index %d (imported %d)."), LodIndex, ImportedLods));
+		}
+		else
+		{
+			int32 TriangleCount = Mesh->GetResourceForRendering()->LODRenderData[LodIndex].GetTotalFaces();
+			if (TriangleCount != ExpectedNumberOfRenderTriangles)
+			{
+				Result.AddError(FString::Printf(TEXT("For LOD %d, expected %d triangles, imported %d."), LodIndex, ExpectedNumberOfRenderTriangles, TriangleCount));
+			}
 		}
 	}
+	else
+	{
+		Result.AddError(FString::Printf(TEXT("No valid render data for this skeletalmesh %s."), *Mesh->GetName()));
+	}
 
 	return Result;
 }
 
-
 FInterchangeTestFunctionResult USkeletalMeshImportTestFunctions::CheckLodCount(USkeletalMesh* Mesh, int32 ExpectedNumberOfLods)
 {
 	FInterchangeTestFunctionResult Result;
-
-	int32 NumLODs = Mesh->GetResourceForRendering()->LODRenderData.Num();
-	if (NumLODs != ExpectedNumberOfLods)
+	if (Mesh->GetResourceForRendering())
+	{
+		int32 NumLODs = Mesh->GetResourceForRendering()->LODRenderData.Num();
+		if (NumLODs != ExpectedNumberOfLods)
+		{
+			Result.AddError(FString::Printf(TEXT("Expected %d LODs, imported %d."), ExpectedNumberOfLods, NumLODs));
+		}
+	}
+	else
 	{
-		Result.AddError(FString::Printf(TEXT("Expected %d LODs, imported %d."), ExpectedNumberOfLods, NumLODs));
+		Result.AddError(FString::Printf(TEXT("No valid render data for this skeletalmesh %s."), *Mesh->GetName()));
 	}
 
 	return Result;
@@ -83,19 +122,25 @@ FInterchangeTestFunctionResult USkeletalMeshImportTestFunctions::CheckMaterialSl
 FInterchangeTestFunctionResult USkeletalMeshImportTestFunctions::CheckSectionCount(USkeletalMesh* Mesh, int32 LodIndex, int32 ExpectedNumberOfSections)
 {
 	FInterchangeTestFunctionResult Result;
-
-	int32 ImportedLods = Mesh->GetResourceForRendering()->LODRenderData.Num();
-	if (LodIndex >= ImportedLods)
+	if (Mesh->GetResourceForRendering())
 	{
-		Result.AddError(FString::Printf(TEXT("The imported mesh doesn't contain LOD index %d (imported %d)."), LodIndex, ImportedLods));
+		int32 ImportedLods = Mesh->GetResourceForRendering()->LODRenderData.Num();
+		if (LodIndex < 0 || LodIndex >= ImportedLods)
+		{
+			Result.AddError(FString::Printf(TEXT("The imported mesh doesn't contain LOD index %d (imported %d)."), LodIndex, ImportedLods));
+		}
+		else
+		{
+			int32 NumSections = Mesh->GetResourceForRendering()->LODRenderData[LodIndex].RenderSections.Num();
+			if (NumSections != ExpectedNumberOfSections)
+			{
+				Result.AddError(FString::Printf(TEXT("For LOD %d, expected %d sections, imported %d."), LodIndex, ExpectedNumberOfSections, NumSections));
+			}
+		}
 	}
 	else
 	{
-		int32 NumSections = Mesh->GetResourceForRendering()->LODRenderData[LodIndex].RenderSections.Num();
-		if (NumSections != ExpectedNumberOfSections)
-		{
-			Result.AddError(FString::Printf(TEXT("For LOD %d, expected %d sections, imported %d."), LodIndex, ExpectedNumberOfSections, NumSections));
-		}
+		Result.AddError(FString::Printf(TEXT("No valid render data for this skeletalmesh %s."), *Mesh->GetName()));
 	}
 
 	return Result;
@@ -105,28 +150,34 @@ FInterchangeTestFunctionResult USkeletalMeshImportTestFunctions::CheckSectionCou
 FInterchangeTestFunctionResult USkeletalMeshImportTestFunctions::CheckTriangleCountInSection(USkeletalMesh* Mesh, int32 LodIndex, int32 SectionIndex, int32 ExpectedNumberOfTriangles)
 {
 	FInterchangeTestFunctionResult Result;
-
-	int32 ImportedLods = Mesh->GetResourceForRendering()->LODRenderData.Num();
-	if (LodIndex >= ImportedLods)
-	{
-		Result.AddError(FString::Printf(TEXT("The imported mesh doesn't contain LOD index %d (imported %d)."), LodIndex, ImportedLods));
-	}
-	else
+	if (Mesh->GetResourceForRendering())
 	{
-		int32 NumSections = Mesh->GetResourceForRendering()->LODRenderData[LodIndex].RenderSections.Num();
-		if (SectionIndex >= NumSections)
+		int32 ImportedLods = Mesh->GetResourceForRendering()->LODRenderData.Num();
+		if (LodIndex < 0 || LodIndex >= ImportedLods)
 		{
-			Result.AddError(FString::Printf(TEXT("The imported mesh doesn't contain section index %d."), SectionIndex));
+			Result.AddError(FString::Printf(TEXT("The imported mesh doesn't contain LOD index %d (imported %d)."), LodIndex, ImportedLods));
 		}
 		else
 		{
-			int32 NumberOfTriangles = Mesh->GetResourceForRendering()->LODRenderData[LodIndex].RenderSections[SectionIndex].NumTriangles;
-			if (NumberOfTriangles != ExpectedNumberOfTriangles)
+			int32 NumSections = Mesh->GetResourceForRendering()->LODRenderData[LodIndex].RenderSections.Num();
+			if (SectionIndex >= NumSections)
 			{
-				Result.AddError(FString::Printf(TEXT("For LOD %d, section index %d, expected %d triangles, imported %d."), LodIndex, SectionIndex, ExpectedNumberOfTriangles, NumberOfTriangles));
+				Result.AddError(FString::Printf(TEXT("The imported mesh doesn't contain section index %d."), SectionIndex));
+			}
+			else
+			{
+				int32 NumberOfTriangles = Mesh->GetResourceForRendering()->LODRenderData[LodIndex].RenderSections[SectionIndex].NumTriangles;
+				if (NumberOfTriangles != ExpectedNumberOfTriangles)
+				{
+					Result.AddError(FString::Printf(TEXT("For LOD %d, section index %d, expected %d triangles, imported %d."), LodIndex, SectionIndex, ExpectedNumberOfTriangles, NumberOfTriangles));
+				}
 			}
 		}
 	}
+	else
+	{
+		Result.AddError(FString::Printf(TEXT("No valid render data for this skeletalmesh %s."), *Mesh->GetName()));
+	}
 
 	return Result;
 }
@@ -135,19 +186,25 @@ FInterchangeTestFunctionResult USkeletalMeshImportTestFunctions::CheckTriangleCo
 FInterchangeTestFunctionResult USkeletalMeshImportTestFunctions::CheckUVChannelCount(USkeletalMesh* Mesh, int32 LodIndex, int32 ExpectedNumberOfUVChannels)
 {
 	FInterchangeTestFunctionResult Result;
-
-	int32 ImportedLods = Mesh->GetResourceForRendering()->LODRenderData.Num();
-	if (LodIndex >= ImportedLods)
+	if (Mesh->GetResourceForRendering())
 	{
-		Result.AddError(FString::Printf(TEXT("The imported mesh doesn't contain LOD index %d (imported %d)."), LodIndex, ImportedLods));
+		int32 ImportedLods = Mesh->GetResourceForRendering()->LODRenderData.Num();
+		if (LodIndex < 0 || LodIndex >= ImportedLods)
+		{
+			Result.AddError(FString::Printf(TEXT("The imported mesh doesn't contain LOD index %d (imported %d)."), LodIndex, ImportedLods));
+		}
+		else
+		{
+			int32 NumUVs = Mesh->GetResourceForRendering()->LODRenderData[LodIndex].GetNumTexCoords();
+			if (NumUVs != ExpectedNumberOfUVChannels)
+			{
+				Result.AddError(FString::Printf(TEXT("For LOD %d, expected %d UVs, imported %d."), LodIndex, ExpectedNumberOfUVChannels, NumUVs));
+			}
+		}
 	}
 	else
 	{
-		int32 NumUVs = Mesh->GetResourceForRendering()->LODRenderData[LodIndex].GetNumTexCoords();
-		if (NumUVs != ExpectedNumberOfUVChannels)
-		{
-			Result.AddError(FString::Printf(TEXT("For LOD %d, expected %d UVs, imported %d."), LodIndex, ExpectedNumberOfUVChannels, NumUVs));
-		}
+		Result.AddError(FString::Printf(TEXT("No valid render data for this skeletalmesh %s."), *Mesh->GetName()));
 	}
 
 	return Result;
@@ -157,73 +214,163 @@ FInterchangeTestFunctionResult USkeletalMeshImportTestFunctions::CheckUVChannelC
 FInterchangeTestFunctionResult USkeletalMeshImportTestFunctions::CheckSectionMaterialName(USkeletalMesh* Mesh, int32 LodIndex, int32 SectionIndex, const FString& ExpectedMaterialName)
 {
 	FInterchangeTestFunctionResult Result;
-
-	int32 ImportedLods = Mesh->GetResourceForRendering()->LODRenderData.Num();
-	if (LodIndex >= ImportedLods)
+	if (Mesh->GetResourceForRendering())
 	{
-		Result.AddError(FString::Printf(TEXT("The imported mesh doesn't contain LOD index %d (imported %d)."), LodIndex, ImportedLods));
+		int32 ImportedLods = Mesh->GetResourceForRendering()->LODRenderData.Num();
+		if (LodIndex < 0 || LodIndex >= ImportedLods)
+		{
+			Result.AddError(FString::Printf(TEXT("The imported mesh doesn't contain LOD index %d (imported %d)."), LodIndex, ImportedLods));
+		}
+		else
+		{
+			const TArray<FSkelMeshRenderSection>& Sections = Mesh->GetResourceForRendering()->LODRenderData[LodIndex].RenderSections;
+			if (SectionIndex >= Sections.Num())
+			{
+				Result.AddError(FString::Printf(TEXT("The imported mesh doesn't contain section index %d (imported %d)."), SectionIndex, Sections.Num()));
+			}
+			else
+			{
+				int32 MaterialIndex = Sections[SectionIndex].MaterialIndex;
+
+				const TArray<FSkeletalMaterial>& SkeletalMaterials = Mesh->GetMaterials();
+				if (!SkeletalMaterials.IsValidIndex(MaterialIndex) || SkeletalMaterials[MaterialIndex].MaterialInterface == nullptr)
+				{
+					Result.AddError(FString::Printf(TEXT("The section references a non-existent material (index %d)."), MaterialIndex));
+				}
+				else
+				{
+					FString MaterialName = SkeletalMaterials[MaterialIndex].MaterialInterface->GetName();
+					if (MaterialName != ExpectedMaterialName)
+					{
+						Result.AddError(FString::Printf(TEXT("For LOD %d section %d, expected material name %s, imported %s."), LodIndex, SectionIndex, *ExpectedMaterialName, *MaterialName));
+					}
+				}
+			}
+		}
 	}
 	else
 	{
-		const TArray<FSkelMeshRenderSection>& Sections = Mesh->GetResourceForRendering()->LODRenderData[LodIndex].RenderSections;
-		if (SectionIndex >= Sections.Num())
+		Result.AddError(FString::Printf(TEXT("No valid render data for this skeletalmesh %s."), *Mesh->GetName()));
+	}
+
+	return Result;
+}
+
+FInterchangeTestFunctionResult USkeletalMeshImportTestFunctions::CheckSectionImportedMaterialSlotName(USkeletalMesh* Mesh, int32 LodIndex, int32 SectionIndex, const FString& ExpectedImportedMaterialSlotName)
+{
+	FInterchangeTestFunctionResult Result;
+	if (Mesh->GetResourceForRendering())
+	{
+		int32 ImportedLods = Mesh->GetResourceForRendering()->LODRenderData.Num();
+		if (LodIndex < 0 || LodIndex >= ImportedLods)
 		{
-			Result.AddError(FString::Printf(TEXT("The imported mesh doesn't contain section index %d (imported %d)."), SectionIndex, Sections.Num()));
+			Result.AddError(FString::Printf(TEXT("The imported mesh doesn't contain LOD index %d (imported %d)."), LodIndex, ImportedLods));
 		}
 		else
 		{
-			int32 MaterialIndex = Sections[SectionIndex].MaterialIndex;
-
-			const TArray<FSkeletalMaterial>& SkeletalMaterials = Mesh->GetMaterials();
-			if (!SkeletalMaterials.IsValidIndex(MaterialIndex) || SkeletalMaterials[MaterialIndex].MaterialInterface == nullptr)
+			const TArray<FSkelMeshRenderSection>& Sections = Mesh->GetResourceForRendering()->LODRenderData[LodIndex].RenderSections;
+			if (SectionIndex >= Sections.Num())
 			{
-				Result.AddError(FString::Printf(TEXT("The section references a non-existent material (index %d)."), MaterialIndex));
+				Result.AddError(FString::Printf(TEXT("The imported mesh doesn't contain section index %d (imported %d)."), SectionIndex, Sections.Num()));
 			}
 			else
 			{
-				FString MaterialName = SkeletalMaterials[MaterialIndex].MaterialInterface->GetName();
-				if (MaterialName != ExpectedMaterialName)
+				int32 MaterialIndex = Sections[SectionIndex].MaterialIndex;
+
+				const TArray<FSkeletalMaterial>& SkeletalMaterials = Mesh->GetMaterials();
+				if (!SkeletalMaterials.IsValidIndex(MaterialIndex))
+				{
+					Result.AddError(FString::Printf(TEXT("The section references a non-existent material (index %d)."), MaterialIndex));
+				}
+				else
 				{
-					Result.AddError(FString::Printf(TEXT("For LOD %d section %d, expected material name %s, imported %s."), LodIndex, SectionIndex, *ExpectedMaterialName, *MaterialName));
+					const FString ImportedMaterialSlotName = SkeletalMaterials[MaterialIndex].ImportedMaterialSlotName.ToString();
+					if (ImportedMaterialSlotName != ExpectedImportedMaterialSlotName)
+					{
+						Result.AddError(FString::Printf(TEXT("For LOD %d section %d, expected imported material slot name %s, imported %s."), LodIndex, SectionIndex, *ExpectedImportedMaterialSlotName, *ImportedMaterialSlotName));
+					}
 				}
 			}
 		}
 	}
+	else
+	{
+		Result.AddError(FString::Printf(TEXT("No valid render data for this skeletalmesh %s."), *Mesh->GetName()));
+	}
 
 	return Result;
 }
 
-
 FInterchangeTestFunctionResult USkeletalMeshImportTestFunctions::CheckVertexIndexPosition(USkeletalMesh* Mesh, int32 LodIndex, int32 VertexIndex, const FVector& ExpectedVertexPosition)
 {
 	FInterchangeTestFunctionResult Result;
-
-	int32 ImportedLods = Mesh->GetResourceForRendering()->LODRenderData.Num();
-	if (LodIndex >= ImportedLods)
+	if (Mesh->GetResourceForRendering())
 	{
-		Result.AddError(FString::Printf(TEXT("The imported mesh doesn't contain LOD index %d (imported %d)."), LodIndex, ImportedLods));
+		int32 ImportedLods = Mesh->GetResourceForRendering()->LODRenderData.Num();
+		if (LodIndex < 0 || LodIndex >= ImportedLods)
+		{
+			Result.AddError(FString::Printf(TEXT("The imported mesh doesn't contain LOD index %d (imported %d)."), LodIndex, ImportedLods));
+		}
+		else
+		{
+			int32 VertexCount = Mesh->GetResourceForRendering()->LODRenderData[LodIndex].StaticVertexBuffers.PositionVertexBuffer.GetNumVertices();
+			if (VertexIndex >= VertexCount)
+			{
+				Result.AddError(FString::Printf(TEXT("The imported mesh doesn't contain vertex index %d (imported %d)."), VertexIndex, VertexCount));
+			}
+			else
+			{
+				FVector VertexPosition = FVector(Mesh->GetResourceForRendering()->LODRenderData[LodIndex].StaticVertexBuffers.PositionVertexBuffer.VertexPosition(VertexIndex));
+				if (!VertexPosition.Equals(ExpectedVertexPosition, UE_DOUBLE_KINDA_SMALL_NUMBER))
+				{
+					Result.AddError(FString::Printf(TEXT("For LOD %d vertex index %d, expected position %s, imported %s."), LodIndex, VertexIndex, *ExpectedVertexPosition.ToString(), *VertexPosition.ToString()));
+				}
+			}
+		}
 	}
 	else
 	{
-		int32 VertexCount = Mesh->GetResourceForRendering()->LODRenderData[LodIndex].StaticVertexBuffers.PositionVertexBuffer.GetNumVertices();
-		if (VertexIndex >= VertexCount)
+		Result.AddError(FString::Printf(TEXT("No valid render data for this skeletalmesh %s."), *Mesh->GetName()));
+	}
+
+	return Result;
+}
+
+FInterchangeTestFunctionResult USkeletalMeshImportTestFunctions::CheckVertexIndexNormal(USkeletalMesh* Mesh, int32 LodIndex, int32 VertexIndex, const FVector& ExpectedVertexNormal)
+{
+	FInterchangeTestFunctionResult Result;
+	if (Mesh->GetResourceForRendering())
+	{
+		int32 ImportedLods = Mesh->GetResourceForRendering()->LODRenderData.Num();
+		if (LodIndex < 0 || LodIndex >= ImportedLods)
 		{
-			Result.AddError(FString::Printf(TEXT("The imported mesh doesn't contain vertex index %d (imported %d)."), VertexIndex, VertexCount));
+			Result.AddError(FString::Printf(TEXT("The imported mesh doesn't contain LOD index %d (imported %d)."), LodIndex, ImportedLods));
 		}
 		else
 		{
-			FVector VertexPosition = FVector(Mesh->GetResourceForRendering()->LODRenderData[LodIndex].StaticVertexBuffers.PositionVertexBuffer.VertexPosition(VertexIndex));
-			if (!VertexPosition.Equals(ExpectedVertexPosition))
+			int32 VertexCount = Mesh->GetResourceForRendering()->LODRenderData[LodIndex].StaticVertexBuffers.StaticMeshVertexBuffer.GetNumVertices();
+			if (VertexIndex >= VertexCount)
+			{
+				Result.AddError(FString::Printf(TEXT("The imported mesh doesn't contain vertex index %d (imported %d)."), VertexIndex, VertexCount));
+			}
+			else
 			{
-				Result.AddError(FString::Printf(TEXT("For LOD %d vertex index %d, expected position %s, imported %s."), LodIndex, VertexIndex, *ExpectedVertexPosition.ToString(), *VertexPosition.ToString()));
+				FVector VertexNormal = FVector(Mesh->GetResourceForRendering()->LODRenderData[LodIndex].StaticVertexBuffers.StaticMeshVertexBuffer.VertexTangentZ(VertexIndex));
+				if (!VertexNormal.Equals(ExpectedVertexNormal, UE_DOUBLE_KINDA_SMALL_NUMBER))
+				{
+					Result.AddError(FString::Printf(TEXT("For LOD %d vertex index %d, expected normal %s, imported %s."), LodIndex, VertexIndex, *ExpectedVertexNormal.ToString(), *VertexNormal.ToString()));
+				}
 			}
 		}
 	}
+	else
+	{
+		Result.AddError(FString::Printf(TEXT("No valid render data for this skeletalmesh %s."), *Mesh->GetName()));
+	}
 
 	return Result;
 }
 
-
 FInterchangeTestFunctionResult USkeletalMeshImportTestFunctions::CheckBoneCount(USkeletalMesh* Mesh, int32 ExpectedNumberOfBones)
 {
 	FInterchangeTestFunctionResult Result;
@@ -268,7 +415,7 @@ FInterchangeTestFunctionResult USkeletalMeshImportTestFunctions::CheckBonePositi
 		else
 		{
 			FVector BonePosition = Mesh->GetRefSkeleton().GetRefBonePose()[BoneIndex].GetLocation();
-			if (!BonePosition.Equals(ExpectedBonePosition))
+			if (!BonePosition.Equals(ExpectedBonePosition, UE_DOUBLE_KINDA_SMALL_NUMBER))
 			{
 				Result.AddError(FString::Printf(TEXT("For bone index %d, expected position %s, imported %s."), BoneIndex, *ExpectedBonePosition.ToString(), *BonePosition.ToString()));
 			}
diff --git a/Engine/Plugins/Tests/InterchangeTests/Source/InterchangeTests/Private/ImportTestFunctions/StaticMeshImportTestFunctions.cpp b/Engine/Plugins/Tests/InterchangeTests/Source/InterchangeTests/Private/ImportTestFunctions/StaticMeshImportTestFunctions.cpp
--- a/Engine/Plugins/Tests/InterchangeTests/Source/InterchangeTests/Private/ImportTestFunctions/StaticMeshImportTestFunctions.cpp
+++ b/Engine/Plugins/Tests/InterchangeTests/Source/InterchangeTests/Private/ImportTestFunctions/StaticMeshImportTestFunctions.cpp
@@ -97,7 +97,7 @@ FInterchangeTestFunctionResult UStaticMeshImportTestFunctions::CheckRenderVertex
 {
 	FInterchangeTestFunctionResult Result;
 
-	if (LodIndex >= Mesh->GetNumLODs())
+	if (LodIndex < 0 || LodIndex >= Mesh->GetNumLODs())
 	{
 		Result.AddError(FString::Printf(TEXT("The imported mesh doesn't contain LOD index %d."), LodIndex));
 	}
@@ -141,6 +141,25 @@ FInterchangeTestFunctionResult UStaticMeshImportTestFunctions::CheckMaterialSlot
 	return Result;
 }
 
+FInterchangeTestFunctionResult UStaticMeshImportTestFunctions::CheckImportedMaterialSlotName(UStaticMesh* Mesh, int32 MaterialIndex, const FString& ExpectedImportedMaterialSlotName)
+{
+	FInterchangeTestFunctionResult Result;
+
+	if (!Mesh->GetStaticMaterials().IsValidIndex(MaterialIndex))
+	{
+		Result.AddError(FString::Printf(TEXT("The imported mesh doesn't have a material at index %d, mesh material count is %d."), MaterialIndex, Mesh->GetStaticMaterials().Num()));
+	}
+	else
+	{
+		const FString ImportedMaterialSlotNameAtIndex = Mesh->GetStaticMaterials()[MaterialIndex].ImportedMaterialSlotName.ToString();
+		if (!ImportedMaterialSlotNameAtIndex.Equals(ExpectedImportedMaterialSlotName, ESearchCase::CaseSensitive))
+		{
+			Result.AddError(FString::Printf(TEXT("For material index %d, expected imported material slot name [%s], imported [%s]."), MaterialIndex, *ExpectedImportedMaterialSlotName, *ImportedMaterialSlotNameAtIndex));
+		}
+	}
+
+	return Result;
+}
 
 FInterchangeTestFunctionResult UStaticMeshImportTestFunctions::CheckPolygonGroupCount(UStaticMesh* Mesh, int32 LodIndex, int32 ExpectedNumberOfPolygonGroups)
 {
@@ -168,7 +187,7 @@ FInterchangeTestFunctionResult UStaticMeshImportTestFunctions::CheckSectionCount
 {
 	FInterchangeTestFunctionResult Result;
 
-	if (LodIndex >= Mesh->GetNumLODs())
+	if (LodIndex < 0 || LodIndex >= Mesh->GetNumLODs())
 	{
 		Result.AddError(FString::Printf(TEXT("The imported mesh doesn't contain LOD index %d."), LodIndex));
 	}
@@ -185,8 +204,40 @@ FInterchangeTestFunctionResult UStaticMeshImportTestFunctions::CheckSectionCount
 	return Result;
 }
 
+FInterchangeTestFunctionResult UStaticMeshImportTestFunctions::CheckRenderVertexIndexNormal(UStaticMesh* Mesh, int32 LodIndex, int32 VertexIndex, const FVector4f& ExpectedVertexNormal)
+{
+	FInterchangeTestFunctionResult Result;
 
-FInterchangeTestFunctionResult UStaticMeshImportTestFunctions::CheckTotalTriangleCount(UStaticMesh* Mesh, int32 LodIndex, int32 ExpectedTotalNumberOfTriangles)
+	if (LodIndex < 0 || LodIndex >= Mesh->GetNumLODs())
+	{
+		Result.AddError(FString::Printf(TEXT("The imported mesh doesn't contain LOD index %d."), LodIndex));
+	}
+	else if (Mesh->GetRenderData() && Mesh->GetRenderData()->LODResources.IsValidIndex(LodIndex))
+	{
+		const int32 NumVertices = Mesh->GetRenderData()->LODResources[LodIndex].GetNumVertices();
+		if (VertexIndex >= NumVertices)
+		{
+			Result.AddError(FString::Printf(TEXT("The imported LOD render data doesn't contain vertex index %d."), VertexIndex));
+		}
+		else
+		{
+			
+			FVector4f VertexNormal = Mesh->GetRenderData()->LODResources[LodIndex].VertexBuffers.StaticMeshVertexBuffer.VertexTangentZ(VertexIndex);
+			if (!VertexNormal.Equals(ExpectedVertexNormal, UE_KINDA_SMALL_NUMBER))
+			{
+				Result.AddError(FString::Printf(TEXT("For LOD %d, vertex index %d, expected vertex normals %s, imported %s."), LodIndex, VertexIndex, *VertexNormal.ToString(), *ExpectedVertexNormal.ToString()));
+			}
+		}
+	}
+	else
+	{
+		Result.AddError(FString::Printf(TEXT("No valid render data For LOD %d."), LodIndex));
+	}
+
+	return Result;
+}
+
+FInterchangeTestFunctionResult UStaticMeshImportTestFunctions::CheckTriangleCount(UStaticMesh* Mesh, int32 LodIndex, int32 ExpectedTotalNumberOfTriangles)
 {
 	FInterchangeTestFunctionResult Result;
 
@@ -207,8 +258,31 @@ FInterchangeTestFunctionResult UStaticMeshImportTestFunctions::CheckTotalTriangl
 	return Result;
 }
 
+FInterchangeTestFunctionResult UStaticMeshImportTestFunctions::CheckRenderTriangleCount(UStaticMesh* Mesh, int32 LodIndex, int32 ExpectedTotalNumberOfTriangles)
+{
+	FInterchangeTestFunctionResult Result;
 
-FInterchangeTestFunctionResult UStaticMeshImportTestFunctions::CheckTotalPolygonCount(UStaticMesh* Mesh, int32 LodIndex, int32 ExpectedTotalNumberOfPolygons)
+	if (LodIndex < 0 || LodIndex >= Mesh->GetNumLODs())
+	{
+		Result.AddError(FString::Printf(TEXT("The imported mesh doesn't contain LOD index %d."), LodIndex));
+	}
+	else if (Mesh->GetRenderData() && Mesh->GetRenderData()->LODResources.IsValidIndex(LodIndex))
+	{
+		const int32 NumTriangles = Mesh->GetRenderData()->LODResources[LodIndex].GetNumTriangles();
+		if (NumTriangles != ExpectedTotalNumberOfTriangles)
+		{
+			Result.AddError(FString::Printf(TEXT("For LOD %d, expected %d triangles, imported %d."), LodIndex, ExpectedTotalNumberOfTriangles, NumTriangles));
+		}
+	}
+	else
+	{
+		Result.AddError(FString::Printf(TEXT("No valid render data For LOD %d."), LodIndex));
+	}
+
+	return Result;
+}
+
+FInterchangeTestFunctionResult UStaticMeshImportTestFunctions::CheckPolygonCount(UStaticMesh* Mesh, int32 LodIndex, int32 ExpectedTotalNumberOfPolygons)
 {
 	FInterchangeTestFunctionResult Result;
 
@@ -322,6 +396,29 @@ FInterchangeTestFunctionResult UStaticMeshImportTestFunctions::CheckUVChannelCou
 	return Result;
 }
 
+FInterchangeTestFunctionResult UStaticMeshImportTestFunctions::CheckRenderUVChannelCount(UStaticMesh* Mesh, int32 LodIndex, int32 ExpectedNumberOfUVChannels)
+{
+	FInterchangeTestFunctionResult Result;
+
+	if (LodIndex < 0 || LodIndex >= Mesh->GetNumLODs())
+	{
+		Result.AddError(FString::Printf(TEXT("The imported mesh doesn't contain LOD index %d."), LodIndex));
+	}
+	else if (Mesh->GetRenderData() && Mesh->GetRenderData()->LODResources.IsValidIndex(LodIndex))
+	{
+		const int32 NumberOfUVChannels = Mesh->GetRenderData()->LODResources[LodIndex].GetNumTexCoords();
+		if (NumberOfUVChannels != ExpectedNumberOfUVChannels)
+		{
+			Result.AddError(FString::Printf(TEXT("Expected %d UV Channels, imported %d."), ExpectedNumberOfUVChannels, NumberOfUVChannels));
+		}
+	}
+	else
+	{
+		Result.AddError(FString::Printf(TEXT("No valid render data For LOD %d."), LodIndex));
+	}
+
+	return Result;
+}
 
 FInterchangeTestFunctionResult UStaticMeshImportTestFunctions::CheckPolygonGroupImportedMaterialSlotName(UStaticMesh* Mesh, int32 LodIndex, int32 PolygonGroupIndex, const FString& ExpectedImportedMaterialSlotName)
 {
@@ -351,11 +448,11 @@ FInterchangeTestFunctionResult UStaticMeshImportTestFunctions::CheckSectionMater
 {
 	FInterchangeTestFunctionResult Result;
 
-	if (LodIndex >= Mesh->GetNumLODs())
+	if (LodIndex < 0 || LodIndex >= Mesh->GetNumLODs())
 	{
 		Result.AddError(FString::Printf(TEXT("The imported mesh doesn't contain LOD index %d."), LodIndex));
 	}
-	else
+	else if (Mesh->GetRenderData() && Mesh->GetRenderData()->LODResources.IsValidIndex(LodIndex))
 	{
 		const FStaticMeshSectionArray& Sections = Mesh->GetRenderData()->LODResources[LodIndex].Sections;
 		if (SectionIndex >= Sections.Num())
@@ -372,6 +469,10 @@ FInterchangeTestFunctionResult UStaticMeshImportTestFunctions::CheckSectionMater
 			}
 		}
 	}
+	else
+	{
+		Result.AddError(FString::Printf(TEXT("No valid render data For LOD %d."), LodIndex));
+	}
 
 	return Result;
 }
@@ -381,11 +482,11 @@ FInterchangeTestFunctionResult UStaticMeshImportTestFunctions::CheckSectionMater
 {
 	FInterchangeTestFunctionResult Result;
 
-	if (LodIndex >= Mesh->GetNumLODs())
+	if (LodIndex < 0 || LodIndex >= Mesh->GetNumLODs())
 	{
 		Result.AddError(FString::Printf(TEXT("The imported mesh doesn't contain LOD index %d."), LodIndex));
 	}
-	else
+	else if (Mesh->GetRenderData() && Mesh->GetRenderData()->LODResources.IsValidIndex(LodIndex))
 	{
 		const FStaticMeshSectionArray& Sections = Mesh->GetRenderData()->LODResources[LodIndex].Sections;
 		if (SectionIndex >= Sections.Num())
@@ -411,10 +512,55 @@ FInterchangeTestFunctionResult UStaticMeshImportTestFunctions::CheckSectionMater
 			}
 		}
 	}
+	else
+	{
+		Result.AddError(FString::Printf(TEXT("No valid render data For LOD %d."), LodIndex));
+	}
 
 	return Result;
 }
 
+FInterchangeTestFunctionResult UStaticMeshImportTestFunctions::CheckSectionImportedMaterialSlotName(UStaticMesh* Mesh, int32 LodIndex, int32 SectionIndex, const FString& ExpectedImportedMaterialSlotName)
+{
+	FInterchangeTestFunctionResult Result;
+
+	if (LodIndex < 0 || LodIndex >= Mesh->GetNumLODs())
+	{
+		Result.AddError(FString::Printf(TEXT("The imported mesh doesn't contain LOD index %d."), LodIndex));
+	}
+	else if (Mesh->GetRenderData() && Mesh->GetRenderData()->LODResources.IsValidIndex(LodIndex))
+	{
+		const FStaticMeshSectionArray& Sections = Mesh->GetRenderData()->LODResources[LodIndex].Sections;
+		if (SectionIndex >= Sections.Num())
+		{
+			Result.AddError(FString::Printf(TEXT("The imported mesh doesn't contain section index %d."), SectionIndex));
+		}
+		else
+		{
+			int32 MaterialIndex = Sections[SectionIndex].MaterialIndex;
+
+			const TArray<FStaticMaterial>& StaticMaterials = Mesh->GetStaticMaterials();
+			if (!StaticMaterials.IsValidIndex(MaterialIndex))
+			{
+				Result.AddError(FString::Printf(TEXT("The section references a non-existent material (index %d)."), MaterialIndex));
+			}
+			else
+			{
+				FString ImportedMaterialSlotName = StaticMaterials[MaterialIndex].ImportedMaterialSlotName.ToString();
+				if (ImportedMaterialSlotName != ExpectedImportedMaterialSlotName)
+				{
+					Result.AddError(FString::Printf(TEXT("For LOD %d section %d, expected imported material slot name %s, imported %s."), LodIndex, SectionIndex, *ExpectedImportedMaterialSlotName, *ImportedMaterialSlotName));
+				}
+			}
+		}
+	}
+	else
+	{
+		Result.AddError(FString::Printf(TEXT("No valid render data For LOD %d."), LodIndex));
+	}
+
+	return Result;
+}
 
 FInterchangeTestFunctionResult UStaticMeshImportTestFunctions::CheckVertexIndexPosition(UStaticMesh* Mesh, int32 LodIndex, int32 VertexIndex, const FVector& ExpectedVertexPosition)
 {
@@ -436,7 +582,7 @@ FInterchangeTestFunctionResult UStaticMeshImportTestFunctions::CheckVertexIndexP
 		else
 		{
 			const FVector VertexPosition = FVector(Attributes.GetVertexPositions()[VertexIndex]);
-			if (!VertexPosition.Equals(ExpectedVertexPosition))
+			if (!VertexPosition.Equals(ExpectedVertexPosition, UE_DOUBLE_KINDA_SMALL_NUMBER))
 			{
 				Result.AddError(FString::Printf(TEXT("For LOD %d vertex index %d, expected position %s, imported %s."), LodIndex, VertexIndex, *ExpectedVertexPosition.ToString(), *VertexPosition.ToString()));
 			}
@@ -544,7 +690,7 @@ FInterchangeTestFunctionResult UStaticMeshImportTestFunctions::CheckSocketLocati
 	}
 
 	FVector SocketLocation = Mesh->Sockets[SocketIndex]->RelativeLocation;
-	if (SocketLocation != ExpectedSocketLocation)
+	if (SocketLocation.Equals(ExpectedSocketLocation, UE_DOUBLE_KINDA_SMALL_NUMBER))
 	{
 		Result.AddError(FString::Printf(TEXT("Expected socket location %s, imported %s"), *ExpectedSocketLocation.ToString(), *SocketLocation.ToString()));
 	}
diff --git a/Engine/Plugins/Tests/InterchangeTests/Source/InterchangeTests/Public/ImportTestFunctions/AnimationImportTestFunctions.h b/Engine/Plugins/Tests/InterchangeTests/Source/InterchangeTests/Public/ImportTestFunctions/AnimationImportTestFunctions.h
--- a/Engine/Plugins/Tests/InterchangeTests/Source/InterchangeTests/Public/ImportTestFunctions/AnimationImportTestFunctions.h
+++ b/Engine/Plugins/Tests/InterchangeTests/Source/InterchangeTests/Public/ImportTestFunctions/AnimationImportTestFunctions.h
@@ -18,4 +18,40 @@ public:
 
 	// UImportTestFunctionsBase interface
 	virtual UClass* GetAssociatedAssetType() const override;
+
+	/** Check whether the expected number of anim sequences are imported */
+	UFUNCTION(Exec)
+	static FInterchangeTestFunctionResult CheckImportedAnimSequenceCount(const TArray<UAnimSequence*>& AnimSequences, int32 ExpectedNumberOfImportedAnimSequences);
+
+	/** Check whether the animation length (second) is as expected */
+	UFUNCTION(Exec)
+	static FInterchangeTestFunctionResult CheckAnimationLength(UAnimSequence* AnimSequence, float ExpectedAnimationLength);
+
+	/** Check whether the animation frame number is as expected */
+	UFUNCTION(Exec)
+	static FInterchangeTestFunctionResult CheckAnimationFrameNumber(UAnimSequence* AnimSequence, int32 ExpectedFrameNumber);
+
+	/** Check whether the given curve key time(sec) for the given curve name has the expected time(sec) */
+	UFUNCTION(Exec)
+	static FInterchangeTestFunctionResult CheckCurveKeyTime(UAnimSequence* AnimSequence, const FString& CurveName, int32 KeyIndex, float ExpectedCurveKeyTime);
+
+	/** Check whether the given curve key value for the given curve name has the expected value */
+	UFUNCTION(Exec)
+	static FInterchangeTestFunctionResult CheckCurveKeyValue(UAnimSequence* AnimSequence, const FString& CurveName, int32 KeyIndex, float ExpectedCurveKeyValue);
+
+	/** Check whether the given curve key arrive tangent for the given curve name has the expected arrive tangent */
+	UFUNCTION(Exec)
+	static FInterchangeTestFunctionResult CheckCurveKeyArriveTangent(UAnimSequence* AnimSequence, const FString& CurveName, int32 KeyIndex, float ExpectedCurveKeyArriveTangent);
+
+	/** Check whether the given curve key arrive tangent weight for the given curve name has the expected arrive tangent weight */
+	UFUNCTION(Exec)
+	static FInterchangeTestFunctionResult CheckCurveKeyArriveTangentWeight(UAnimSequence* AnimSequence, const FString& CurveName, int32 KeyIndex, float ExpectedCurveKeyArriveTangentWeight);
+
+	/** Check whether the given curve key leave tangent for the given curve name has the expected leave tangent */
+	UFUNCTION(Exec)
+	static FInterchangeTestFunctionResult CheckCurveKeyLeaveTangent(UAnimSequence* AnimSequence, const FString& CurveName, int32 KeyIndex, float ExpectedCurveKeyLeaveTangent);
+
+	/** Check whether the given curve key leave tangent weight for the given curve name has the expected leave tangent weight */
+	UFUNCTION(Exec)
+	static FInterchangeTestFunctionResult CheckCurveKeyLeaveTangentWeight(UAnimSequence* AnimSequence, const FString& CurveName, int32 KeyIndex, float ExpectedCurveKeyLeaveTangentWeight);
 };
diff --git a/Engine/Plugins/Tests/InterchangeTests/Source/InterchangeTests/Public/ImportTestFunctions/SkeletalMeshImportTestFunctions.h b/Engine/Plugins/Tests/InterchangeTests/Source/InterchangeTests/Public/ImportTestFunctions/SkeletalMeshImportTestFunctions.h
--- a/Engine/Plugins/Tests/InterchangeTests/Source/InterchangeTests/Public/ImportTestFunctions/SkeletalMeshImportTestFunctions.h
+++ b/Engine/Plugins/Tests/InterchangeTests/Source/InterchangeTests/Public/ImportTestFunctions/SkeletalMeshImportTestFunctions.h
@@ -28,6 +28,10 @@ public:
 	UFUNCTION(Exec)
 	static FInterchangeTestFunctionResult CheckRenderVertexCount(USkeletalMesh* Mesh, int32 LodIndex, int32 ExpectedNumberOfRenderVertices);
 
+	/** Check whether the triangle count in the built render data for the given LOD is as expected */
+	UFUNCTION(Exec)
+	static FInterchangeTestFunctionResult CheckRenderTriangleCount(USkeletalMesh* Mesh, int32 LodIndex, int32 ExpectedNumberOfRenderTriangles);
+
 	/** Check whether the mesh has the expected number of LODs */
 	UFUNCTION(Exec)
 	static FInterchangeTestFunctionResult CheckLodCount(USkeletalMesh* Mesh, int32 ExpectedNumberOfLods);
@@ -52,10 +56,18 @@ public:
 	UFUNCTION(Exec)
 	static FInterchangeTestFunctionResult CheckSectionMaterialName(USkeletalMesh* Mesh, int32 LodIndex, int32 SectionIndex, const FString& ExpectedMaterialName);
 
+	/** Check whether the imported material slot name for the given section in the render data for the given LOD is as expected */
+	UFUNCTION(Exec)
+	static FInterchangeTestFunctionResult CheckSectionImportedMaterialSlotName(USkeletalMesh* Mesh, int32 LodIndex, int32 SectionIndex, const FString& ExpectedImportedMaterialSlotName);
+
 	/** Check whether the vertex of the given index is at the expected position */
 	UFUNCTION(Exec)
 	static FInterchangeTestFunctionResult CheckVertexIndexPosition(USkeletalMesh* Mesh, int32 LodIndex, int32 VertexIndex, const FVector& ExpectedVertexPosition);
 
+	/** Check whether the vertex of the given index is at the expected normal */
+	UFUNCTION(Exec)
+	static FInterchangeTestFunctionResult CheckVertexIndexNormal(USkeletalMesh* Mesh, int32 LodIndex, int32 VertexIndex, const FVector& ExpectedVertexNormal);
+
 	/** Check whether the mesh has the expected number of bones */
 	UFUNCTION(Exec)
 	static FInterchangeTestFunctionResult CheckBoneCount(USkeletalMesh* Mesh, int32 ExpectedNumberOfBones);
diff --git a/Engine/Plugins/Tests/InterchangeTests/Source/InterchangeTests/Public/ImportTestFunctions/StaticMeshImportTestFunctions.h b/Engine/Plugins/Tests/InterchangeTests/Source/InterchangeTests/Public/ImportTestFunctions/StaticMeshImportTestFunctions.h
--- a/Engine/Plugins/Tests/InterchangeTests/Source/InterchangeTests/Public/ImportTestFunctions/StaticMeshImportTestFunctions.h
+++ b/Engine/Plugins/Tests/InterchangeTests/Source/InterchangeTests/Public/ImportTestFunctions/StaticMeshImportTestFunctions.h
@@ -38,87 +38,107 @@ public:
 	UFUNCTION(Exec)
 	static FInterchangeTestFunctionResult CheckImportedStaticMeshCount(const TArray<UStaticMesh*>& Meshes, int32 ExpectedNumberOfImportedStaticMeshes);
 
-	/** Check whether the vertex count for the given LOD is as expected */
+	/** Check whether the static mesh has the expected number of LODs */
 	UFUNCTION(Exec)
-	static FInterchangeTestFunctionResult CheckVertexCount(UStaticMesh* Mesh, int32 LodIndex, int32 ExpectedNumberOfVertices);
+	static FInterchangeTestFunctionResult CheckLodCount(UStaticMesh* Mesh, int32 ExpectedNumberOfLods);
 
-	/** Check whether the vertex count in the built render data for the given LOD is as expected */
+	/** Check whether the render data for the given mesh LOD has the expected number of sections */
 	UFUNCTION(Exec)
-	static FInterchangeTestFunctionResult CheckRenderVertexCount(UStaticMesh* Mesh, int32 LodIndex, int32 ExpectedNumberOfRenderVertices);
+	static FInterchangeTestFunctionResult CheckSectionCount(UStaticMesh* Mesh, int32 LodIndex, int32 ExpectedNumberOfSections);
 
-	/** Check whether the mesh has the expected number of LODs */
+	/** Check whether the static mesh has the expected number of material slots */
 	UFUNCTION(Exec)
-	static FInterchangeTestFunctionResult CheckLodCount(UStaticMesh* Mesh, int32 ExpectedNumberOfLods);
+	static FInterchangeTestFunctionResult CheckMaterialSlotCount(UStaticMesh* Mesh, int32 ExpectedNumberOfMaterialSlots);
 
-	/** Check whether the mesh has the expected number of material slots */
+	/** Check whether the static mesh imported material slot name for the given material index is as expected */
 	UFUNCTION(Exec)
-	static FInterchangeTestFunctionResult CheckMaterialSlotCount(UStaticMesh* Mesh, int32 ExpectedNumberOfMaterialSlots);
+	static FInterchangeTestFunctionResult CheckImportedMaterialSlotName(UStaticMesh* Mesh, int32 MaterialIndex, const FString& ExpectedImportedMaterialSlotName);
 
-	/** Check whether the mesh has the expected number of polygon groups for the given LOD */
+	/** Check whether the mesh description has the expected number of polygon groups for the given LOD */
 	UFUNCTION(Exec)
 	static FInterchangeTestFunctionResult CheckPolygonGroupCount(UStaticMesh* Mesh, int32 LodIndex, int32 ExpectedNumberOfPolygonGroups);
 
-	/** Check whether the built render data for the given mesh LOD has the expected number of sections */
+	/** Check whether the mesh description material slot name for the given polygon group in the given LOD is as expected */
 	UFUNCTION(Exec)
-	static FInterchangeTestFunctionResult CheckSectionCount(UStaticMesh* Mesh, int32 LodIndex, int32 ExpectedNumberOfSections);
+	static FInterchangeTestFunctionResult CheckPolygonGroupImportedMaterialSlotName(UStaticMesh* Mesh, int32 LodIndex, int32 PolygonGroupIndex, const FString& ExpectedImportedMaterialSlotName);
 
-	/** Check whether the mesh has the expected number of triangles for the given LOD */
+	/** Check whether the section index in the built render data for the given LOD is referencing the expected material index */
 	UFUNCTION(Exec)
-	static FInterchangeTestFunctionResult CheckTotalTriangleCount(UStaticMesh* Mesh, int32 LodIndex, int32 ExpectedTotalNumberOfTriangles);
+	static FInterchangeTestFunctionResult CheckSectionMaterialIndex(UStaticMesh* Mesh, int32 LodIndex, int32 SectionIndex, int32 ExpectedMaterialIndex);
 
-	/** Check whether the mesh has the expected number of polygons for the given LOD */
+	/** Check whether the section index in the built render data for the given LOD is referencing the expected material */
 	UFUNCTION(Exec)
-	static FInterchangeTestFunctionResult CheckTotalPolygonCount(UStaticMesh* Mesh, int32 LodIndex, int32 ExpectedNumberOfPolygons);
+	static FInterchangeTestFunctionResult CheckSectionMaterialName(UStaticMesh* Mesh, int32 LodIndex, int32 SectionIndex, const FString& ExpectedMaterialName);
 
-	/** Check whether the mesh contains any quads or ngons */
+	/** Check whether the section index in the built render data for the given LOD is referencing the expected imported material slot name*/
 	UFUNCTION(Exec)
-	static FInterchangeTestFunctionResult CheckThatMeshHasQuadsOrNgons(UStaticMesh* Mesh, int32 LodIndex, bool bMeshHasQuadsOrNgons);
+	static FInterchangeTestFunctionResult CheckSectionImportedMaterialSlotName(UStaticMesh* Mesh, int32 LodIndex, int32 SectionIndex, const FString& ExpectedImportedMaterialSlotName);
 
-	/** Check whether the given polygon group of the given LOD has the expected number of triangles */
+	/** Check whether the mesh description vertex count for the given LOD is as expected */
 	UFUNCTION(Exec)
-	static FInterchangeTestFunctionResult CheckTriangleCountInPolygonGroup(UStaticMesh* Mesh, int32 LodIndex, int32 PolygonGroupIndex, int32 ExpectedNumberOfTriangles);
+	static FInterchangeTestFunctionResult CheckVertexCount(UStaticMesh* Mesh, int32 LodIndex, int32 ExpectedNumberOfVertices);
 
-	/** Check whether the given polygon group of the given LOD has the expected number of polygons */
+	/** Check whether the render data vertex count for the given LOD is as expected */
 	UFUNCTION(Exec)
-	static FInterchangeTestFunctionResult CheckPolygonCountInPolygonGroup(UStaticMesh* Mesh, int32 LodIndex, int32 PolygonGroupIndex, int32 ExpectedNumberOfPolygons);
+	static FInterchangeTestFunctionResult CheckRenderVertexCount(UStaticMesh* Mesh, int32 LodIndex, int32 ExpectedNumberOfRenderVertices);
 
-	/** Check whether the mesh has the expected number of UV channels */
+	/** Check whether the mesh description vertex of the given index is at the expected position */
 	UFUNCTION(Exec)
-	static FInterchangeTestFunctionResult CheckUVChannelCount(UStaticMesh* Mesh, int32 LodIndex, int32 ExpectedNumberOfUVChannels);
+	static FInterchangeTestFunctionResult CheckVertexIndexPosition(UStaticMesh* Mesh, int32 LodIndex, int32 VertexIndex, const FVector& ExpectedVertexPosition);
 
-	/** Check whether the material slot name for the given polygon group in the given LOD is as expected */
+	/** Check whether the render data vertex index of the given LOD has the expected vertex normal */
 	UFUNCTION(Exec)
-	static FInterchangeTestFunctionResult CheckPolygonGroupImportedMaterialSlotName(UStaticMesh* Mesh, int32 LodIndex, int32 PolygonGroupIndex, const FString& ExpectedImportedMaterialSlotName);
+	static FInterchangeTestFunctionResult CheckRenderVertexIndexNormal(UStaticMesh* Mesh, int32 LodIndex, int32 VertexIndex, const FVector4f& ExpectedVertexNormal);
 
-	/** Check whether the section index in the built render data for the given LOD is referencing the expected material index */
+	/** Check whether the mesh description has the expected number of triangles for the given LOD */
 	UFUNCTION(Exec)
-	static FInterchangeTestFunctionResult CheckSectionMaterialIndex(UStaticMesh* Mesh, int32 LodIndex, int32 SectionIndex, int32 ExpectedMaterialIndex);
+	static FInterchangeTestFunctionResult CheckTriangleCount(UStaticMesh* Mesh, int32 LodIndex, int32 ExpectedTotalNumberOfTriangles);
 
-	/** Check whether the section index in the built render data for the given LOD is referencing the expected material */
+	/** Check whether the render data has the expected number of triangles for the given LOD */
 	UFUNCTION(Exec)
-	static FInterchangeTestFunctionResult CheckSectionMaterialName(UStaticMesh* Mesh, int32 LodIndex, int32 SectionIndex, const FString& ExpectedMaterialName);
+	static FInterchangeTestFunctionResult CheckRenderTriangleCount(UStaticMesh* Mesh, int32 LodIndex, int32 ExpectedTotalNumberOfTriangles);
 
-	/** Check whether the vertex of the given index is at the expected position */
+	/** Check whether the mesh description has the expected number of polygons for the given LOD */
 	UFUNCTION(Exec)
-	static FInterchangeTestFunctionResult CheckVertexIndexPosition(UStaticMesh* Mesh, int32 LodIndex, int32 VertexIndex, const FVector& ExpectedVertexPosition);
+	static FInterchangeTestFunctionResult CheckPolygonCount(UStaticMesh* Mesh, int32 LodIndex, int32 ExpectedNumberOfPolygons);
+
+	/** Check whether the mesh description contains any quads or ngons */
+	UFUNCTION(Exec)
+	static FInterchangeTestFunctionResult CheckThatMeshHasQuadsOrNgons(UStaticMesh* Mesh, int32 LodIndex, bool bMeshHasQuadsOrNgons);
+
+	/** Check whether the given mesh description polygon group of the given LOD has the expected number of triangles */
+	UFUNCTION(Exec)
+	static FInterchangeTestFunctionResult CheckTriangleCountInPolygonGroup(UStaticMesh* Mesh, int32 LodIndex, int32 PolygonGroupIndex, int32 ExpectedNumberOfTriangles);
+
+	/** Check whether the given mesh description polygon group of the given LOD has the expected number of polygons */
+	UFUNCTION(Exec)
+	static FInterchangeTestFunctionResult CheckPolygonCountInPolygonGroup(UStaticMesh* Mesh, int32 LodIndex, int32 PolygonGroupIndex, int32 ExpectedNumberOfPolygons);
+
+	/** Check whether the mesh description has the expected number of UV channels */
+	UFUNCTION(Exec)
+	static FInterchangeTestFunctionResult CheckUVChannelCount(UStaticMesh* Mesh, int32 LodIndex, int32 ExpectedNumberOfUVChannels);
+
+	/** Check whether the render data has the expected number of UV channels */
+	UFUNCTION(Exec)
+	static FInterchangeTestFunctionResult CheckRenderUVChannelCount(UStaticMesh* Mesh, int32 LodIndex, int32 ExpectedNumberOfUVChannels);
 
-	/** Check whether the expected number of simple collision primitives were imported */
+	/** Check whether the static mesh expected number of simple collision primitives were imported */
 	UFUNCTION(Exec)
 	static FInterchangeTestFunctionResult CheckSimpleCollisionPrimitiveCount(UStaticMesh* Mesh, int32 ExpectedSphereElementCount, int32 ExpectedBoxElementCount, int32 ExpectedCapsuleElementCount, int32 ExpectedConvexElementCount, int32 ExpectedTaperedCapsuleElementCount);
 
-	/** Check whether the expected number of sockets were imported */
+	/** Check whether the static mesh expected number of sockets were imported */
 	UFUNCTION(Exec)
 	static FInterchangeTestFunctionResult CheckSocketCount(UStaticMesh* Mesh, int32 ExpectedSocketCount);
 
-	/** Check whether the given socket index has the expected name */
+	/** Check whether the static mesh given socket index has the expected name */
 	UFUNCTION(Exec)
 	static FInterchangeTestFunctionResult CheckSocketName(UStaticMesh* Mesh, int32 SocketIndex, const FString& ExpectedSocketName);
 
-	/** Check whether the given socket index has the expected location */
+	/** Check whether the static mesh given socket index has the expected location */
 	UFUNCTION(Exec)
 	static FInterchangeTestFunctionResult CheckSocketLocation(UStaticMesh* Mesh, int32 SocketIndex, const FVector& ExpectedSocketLocation);
 
-	/** Check whether the mesh is equivalent to a ground truth asset */
+	/** Check whether the static mesh is equivalent to a ground truth asset */
 	UFUNCTION(Exec)
 	static FInterchangeTestFunctionResult CheckAgainstGroundTruth(UStaticMesh* Mesh, TSoftObjectPtr<UStaticMesh> MeshToCompare,
 		bool bCheckVertexCountEqual = true,
@@ -129,11 +149,11 @@ public:
 		bool bCheckNormalsEqual = true
 	);
 
-	/** Check whether the mesh source model has the expected build settings */
+	/** Check whether the static mesh source model has the expected build settings */
 	UFUNCTION(Exec)
 	static FInterchangeTestFunctionResult CheckBuildSettings(UStaticMesh* Mesh, int32 LodIndex, const FMeshBuildSettings& ExpectedBuildSettings);
 
-	/** Check whether the mesh has the expected Nanite settings */
+	/** Check whether the static mesh has the expected Nanite settings */
 	UFUNCTION(Exec)
 	static FInterchangeTestFunctionResult CheckNaniteSettings(UStaticMesh* Mesh, const FMeshNaniteSettings& ExpectedNaniteSettings);
 };
