Add 12 new test functions to improve capabilities of interchange test plans:
- CheckActorPropertyValue, CheckComponentPropertyValue, CheckActorClassCount, CheckActorClass in UActorImportTestFunctions
- CheckShadingModel, CheckBlendMode, CheckIsTwoSided, CheckOpacityMaskClipValue in UMaterialImportTestFunctions
- CheckImportedLevelVariantSetsCount, CheckVariantSetsCount, CheckVariantsCount, CheckBindingsCount in ULevelVariantSetsImportTestFunctions

diff --git a/Engine/Plugins/Tests/InterchangeTests/InterchangeTests.uplugin b/Engine/Plugins/Tests/InterchangeTests/InterchangeTests.uplugin
--- a/Engine/Plugins/Tests/InterchangeTests/InterchangeTests.uplugin
+++ b/Engine/Plugins/Tests/InterchangeTests/InterchangeTests.uplugin
@@ -35,6 +35,10 @@
 				"Editor",
 				"Program"
 			]
+		},
+		{
+			"Name": "VariantManagerContent",
+			"Enabled": true
 		}
 	]
 }
diff --git a/Engine/Plugins/Tests/InterchangeTests/Source/InterchangeTests/InterchangeTests.Build.cs b/Engine/Plugins/Tests/InterchangeTests/Source/InterchangeTests/InterchangeTests.Build.cs
--- a/Engine/Plugins/Tests/InterchangeTests/Source/InterchangeTests/InterchangeTests.Build.cs
+++ b/Engine/Plugins/Tests/InterchangeTests/Source/InterchangeTests/InterchangeTests.Build.cs
@@ -29,7 +29,8 @@ namespace UnrealBuildTool.Rules
 					"InterchangeEngine",
 					"MeshDescription",
 					"StaticMeshDescription",
-					"UnrealEd"
+					"UnrealEd",
+					"VariantManagerContent"
 				}
 			);
 		}
diff --git a/Engine/Plugins/Tests/InterchangeTests/Source/InterchangeTests/Private/ImportTestFunctions/ActorImportTestFunctions.cpp b/Engine/Plugins/Tests/InterchangeTests/Source/InterchangeTests/Private/ImportTestFunctions/ActorImportTestFunctions.cpp
--- a/Engine/Plugins/Tests/InterchangeTests/Source/InterchangeTests/Private/ImportTestFunctions/ActorImportTestFunctions.cpp
+++ b/Engine/Plugins/Tests/InterchangeTests/Source/InterchangeTests/Private/ImportTestFunctions/ActorImportTestFunctions.cpp
@@ -4,9 +4,37 @@
 #include "GameFramework/Actor.h"
 #include "ImportTestFunctions/ImportTestFunctionsBase.h"
 #include "InterchangeTestFunction.h"
+#include "Internationalization/Regex.h"
 
 #include UE_INLINE_GENERATED_CPP_BY_NAME(ActorImportTestFunctions)
 
+namespace UE::Interchange::Tests::Private
+{
+	bool TryGetPropertyTextValue(const UObject* Object, const FString& PropertyName, FString& OutTextValue)
+	{
+		const FProperty* Property = Object->GetClass()->FindPropertyByName(*PropertyName);
+		if (Property == nullptr)
+		{
+			return false;
+		}
+
+		Property->ExportTextItem_InContainer(OutTextValue, Object, nullptr, nullptr, PPF_None);
+		return true;
+	}
+
+	const UActorComponent* FindActorComponentByName(const AActor* Actor, const FName& ComponentName)
+	{
+		for (const UActorComponent* Component : Actor->GetComponents())
+		{
+			if (Component->GetFName() == ComponentName)
+			{
+				return Component;
+			}
+		}
+
+		return nullptr;
+	}
+}
 
 UClass* UActorImportTestFunctions::GetAssociatedAssetType() const
 {
@@ -24,3 +52,96 @@ FInterchangeTestFunctionResult UActorImportTestFunctions::CheckImportedActorCoun
 	return Result;
 }
 
+FInterchangeTestFunctionResult UActorImportTestFunctions::CheckActorClassCount(const TArray<AActor*>& Actors, TSubclassOf<AActor> Class, int32 ExpectedNumberOfActors)
+{
+	FInterchangeTestFunctionResult Result;
+	const int32 ImportedNumberOfActors = Actors.FilterByPredicate([Class](const AActor* Actor){ return Actor->GetClass() == Class; }).Num();
+
+	if (ImportedNumberOfActors != ExpectedNumberOfActors)
+	{
+		Result.AddError(FString::Printf(TEXT("Expected %d actors of class %s, imported %d."), ExpectedNumberOfActors, *Class->GetName(), ExpectedNumberOfActors));
+	}
+
+	return Result;
+}
+
+FInterchangeTestFunctionResult UActorImportTestFunctions::CheckActorClass(const AActor* Actor, TSubclassOf<AActor> ExpectedClass)
+{
+	FInterchangeTestFunctionResult Result;
+	const UClass* ImportedClass = Actor->GetClass();
+
+	if (ImportedClass != ExpectedClass)
+	{
+		Result.AddError(FString::Printf(TEXT("Expected actor of class %s, imported %s."), *ExpectedClass->GetName(), *ImportedClass->GetName()));
+	}
+
+	return Result;
+}
+
+FInterchangeTestFunctionResult UActorImportTestFunctions::CheckActorPropertyValue(const AActor* Actor, const FString& PropertyName, bool bUseRegexToMatchValue, const FString& ExpectedValue)
+{
+	using namespace UE::Interchange::Tests::Private;
+
+	FInterchangeTestFunctionResult Result;
+	FString ImportedTextValue;
+
+	if (!TryGetPropertyTextValue(Actor, PropertyName, ImportedTextValue))
+	{
+		Result.AddError(FString::Printf(TEXT("The imported actor doesn't contain property '%s'."), *PropertyName));
+	}
+	else
+	{
+		if (bUseRegexToMatchValue)
+		{
+			FRegexMatcher RegexMatcher(FRegexPattern(ExpectedValue), ImportedTextValue);
+			if (!RegexMatcher.FindNext())
+			{
+				Result.AddError(FString::Printf(TEXT("For actor property '%s', expected value to match regex pattern '%s', imported '%s'."), *PropertyName, *ExpectedValue, *ImportedTextValue));
+			}
+		}
+		else if (ImportedTextValue != ExpectedValue)
+		{
+			Result.AddError(FString::Printf(TEXT("For actor property '%s', expected value '%s', imported '%s'."), *PropertyName, *ExpectedValue, *ImportedTextValue));
+		}
+	}
+
+	return Result;
+}
+
+FInterchangeTestFunctionResult UActorImportTestFunctions::CheckComponentPropertyValue(const AActor* Actor, const FString& ComponentName, const FString& PropertyName, bool bUseRegexToMatchValue, const FString& ExpectedValue)
+{
+	using namespace UE::Interchange::Tests::Private;
+
+	FInterchangeTestFunctionResult Result;
+	const UActorComponent* Component = FindActorComponentByName(Actor, *ComponentName);
+
+	if (Component == nullptr)
+	{
+		Result.AddError(FString::Printf(TEXT("The imported actor doesn't contain component '%s'."), *ComponentName));
+	}
+	else
+	{
+		FString ImportedTextValue;
+		if (!TryGetPropertyTextValue(Component, PropertyName, ImportedTextValue))
+		{
+			Result.AddError(FString::Printf(TEXT("The imported actor component '%s' doesn't contain property '%s'."), *ComponentName, *PropertyName));
+		}
+		else
+		{
+			if (bUseRegexToMatchValue)
+			{
+				FRegexMatcher RegexMatcher(FRegexPattern(ExpectedValue), ImportedTextValue);
+				if (!RegexMatcher.FindNext())
+				{
+					Result.AddError(FString::Printf(TEXT("For property '%s' in actor component '%s', expected value to match regex pattern '%s', imported '%s'."), *PropertyName, *ComponentName, *ExpectedValue, *ImportedTextValue));
+				}
+			}
+			else if (ImportedTextValue != ExpectedValue)
+			{
+				Result.AddError(FString::Printf(TEXT("For property '%s' in actor component '%s', expected value '%s', imported '%s'."), *PropertyName, *ComponentName, *ExpectedValue, *ImportedTextValue));
+			}
+		}
+	}
+
+	return Result;
+}
diff --git a/Engine/Plugins/Tests/InterchangeTests/Source/InterchangeTests/Private/ImportTestFunctions/LevelVariantSetsImportTestFunctions.cpp b/Engine/Plugins/Tests/InterchangeTests/Source/InterchangeTests/Private/ImportTestFunctions/LevelVariantSetsImportTestFunctions.cpp
new file mode 100644
--- /dev/null
+++ b/Engine/Plugins/Tests/InterchangeTests/Source/InterchangeTests/Private/ImportTestFunctions/LevelVariantSetsImportTestFunctions.cpp
@@ -0,0 +1,90 @@
+// Copyright Epic Games, Inc. All Rights Reserved.
+
+#include "ImportTestFunctions/LevelVariantSetsImportTestFunctions.h"
+#include "LevelVariantSets.h"
+#include "ImportTestFunctions/ImportTestFunctionsBase.h"
+#include "InterchangeTestFunction.h"
+#include "Variant.h"
+#include "VariantSet.h"
+
+#include UE_INLINE_GENERATED_CPP_BY_NAME(LevelVariantSetsImportTestFunctions)
+
+
+UClass* ULevelVariantSetsImportTestFunctions::GetAssociatedAssetType() const
+{
+	return ULevelVariantSets::StaticClass();
+}
+
+FInterchangeTestFunctionResult ULevelVariantSetsImportTestFunctions::CheckLevelVariantSetsCount(const TArray<ULevelVariantSets*>& LevelVariantSetsAssets, int32 ExpectedNumberOfLevelVariantSets)
+{
+	FInterchangeTestFunctionResult Result;
+	if (LevelVariantSetsAssets.Num() != ExpectedNumberOfLevelVariantSets)
+	{
+		Result.AddError(FString::Printf(TEXT("Expected %d level variant sets, imported %d."), ExpectedNumberOfLevelVariantSets, LevelVariantSetsAssets.Num()));
+	}
+
+	return Result;
+}
+
+FInterchangeTestFunctionResult ULevelVariantSetsImportTestFunctions::CheckVariantSetsCount(ULevelVariantSets* LevelVariantSets, int32 ExpectedNumberOfVariantSets)
+{
+	FInterchangeTestFunctionResult Result;
+	const int32 ImportedNumberOfVariantSets = LevelVariantSets->GetNumVariantSets();
+
+	if (ImportedNumberOfVariantSets!= ExpectedNumberOfVariantSets)
+	{
+		Result.AddError(FString::Printf(TEXT("Expected %d variant sets, imported %d."), ExpectedNumberOfVariantSets, ImportedNumberOfVariantSets));
+	}
+
+	return Result;
+}
+
+FInterchangeTestFunctionResult ULevelVariantSetsImportTestFunctions::CheckVariantsCount(ULevelVariantSets* LevelVariantSets, const FString& VariantSetName, int32 ExpectedNumberOfVariants)
+{
+	FInterchangeTestFunctionResult Result;
+	const UVariantSet* VariantSet = LevelVariantSets->GetVariantSetByName(VariantSetName);
+
+	if (VariantSet == nullptr)
+	{
+		Result.AddError(FString::Printf(TEXT("The imported level variant sets doesn't contain variant set '%s'."), *VariantSetName));
+	}
+	else
+	{
+		const int32 ImportedNumberOfVariants = VariantSet->GetNumVariants();
+		if (ImportedNumberOfVariants != ExpectedNumberOfVariants)
+		{
+			Result.AddError(FString::Printf(TEXT("For variant set '%s', expected %d variants, imported %d."), *VariantSetName, ExpectedNumberOfVariants, ImportedNumberOfVariants));
+		}
+	}
+
+	return Result;
+}
+
+FInterchangeTestFunctionResult ULevelVariantSetsImportTestFunctions::CheckBindingsCount(ULevelVariantSets* LevelVariantSets, const FString& VariantSetName, const FString& VariantName, int32 ExpectedNumberOfBindings)
+{
+	FInterchangeTestFunctionResult Result;
+	UVariantSet* VariantSet = LevelVariantSets->GetVariantSetByName(VariantSetName);
+
+	if (VariantSet == nullptr)
+	{
+		Result.AddError(FString::Printf(TEXT("The imported level variant sets doesn't contain variant set '%s'."), *VariantSetName));
+	}
+	else
+	{
+		const UVariant* Variant = VariantSet->GetVariantByName(VariantName);
+		if (Variant == nullptr)
+		{
+			Result.AddError(FString::Printf(TEXT("The imported variant set '%s' doesn't contain variant '%s'."), *VariantSetName, *VariantName));
+		}
+		else
+		{
+			const int32 ImportedNumberOfBindings = Variant->GetBindings().Num();
+			if (ImportedNumberOfBindings != ExpectedNumberOfBindings)
+			{
+				Result.AddError(FString::Printf(TEXT("For variant '%s' in set '%s', expected %d bindings, imported %d."), *VariantName, *VariantSetName, ExpectedNumberOfBindings, ImportedNumberOfBindings));
+			}
+		}
+	}
+
+	return Result;
+}
diff --git a/Engine/Plugins/Tests/InterchangeTests/Source/InterchangeTests/Private/ImportTestFunctions/MaterialImportTestFunctions.cpp b/Engine/Plugins/Tests/InterchangeTests/Source/InterchangeTests/Private/ImportTestFunctions/MaterialImportTestFunctions.cpp
--- a/Engine/Plugins/Tests/InterchangeTests/Source/InterchangeTests/Private/ImportTestFunctions/MaterialImportTestFunctions.cpp
+++ b/Engine/Plugins/Tests/InterchangeTests/Source/InterchangeTests/Private/ImportTestFunctions/MaterialImportTestFunctions.cpp
@@ -62,3 +62,69 @@ FInterchangeTestFunctionResult UMaterialImportTestFunctions::CheckImportedMateri
 	return Result;
 }
 
+FInterchangeTestFunctionResult UMaterialImportTestFunctions::CheckShadingModel(const UMaterialInterface* MaterialInterface, EMaterialShadingModel ExpectedShadingModel)
+{
+	FInterchangeTestFunctionResult Result;
+	const FMaterialShadingModelField ImportedShadingModels = MaterialInterface->GetShadingModels();
+
+	const int32 NumImportedShadingModels = ImportedShadingModels.CountShadingModels();
+	if (NumImportedShadingModels != 1)
+	{
+		Result.AddError(FString::Printf(TEXT("Expected only 1 shading model, imported %d."), NumImportedShadingModels));
+	}
+	else
+	{
+		const EMaterialShadingModel ImportedShadingModel = ImportedShadingModels.GetFirstShadingModel();
+		if (ImportedShadingModel != ExpectedShadingModel)
+		{
+			const FString ImportedDisplayValue = UEnum::GetDisplayValueAsText(ImportedShadingModel).ToString();
+			const FString ExpectedDisplayValue = UEnum::GetDisplayValueAsText(ExpectedShadingModel).ToString();
+			Result.AddError(FString::Printf(TEXT("Expected shading model %s, imported %s."), *ExpectedDisplayValue, *ImportedDisplayValue));
+		}
+	}
+
+	return Result;
+}
+
+FInterchangeTestFunctionResult UMaterialImportTestFunctions::CheckBlendMode(const UMaterialInterface* MaterialInterface, EBlendMode ExpectedBlendMode)
+{
+	FInterchangeTestFunctionResult Result;
+	const EBlendMode ImportedBlendMode = MaterialInterface->GetBlendMode();
+
+	if (ImportedBlendMode != ExpectedBlendMode)
+	{
+		const FString ImportedDisplayValue = UEnum::GetDisplayValueAsText(ImportedBlendMode).ToString();
+		const FString ExpectedDisplayValue = UEnum::GetDisplayValueAsText(ExpectedBlendMode).ToString();
+		Result.AddError(FString::Printf(TEXT("Expected blend mode %s, imported %s."), *ExpectedDisplayValue, *ImportedDisplayValue));
+	}
+
+	return Result;
+}
+
+FInterchangeTestFunctionResult UMaterialImportTestFunctions::CheckIsTwoSided(const UMaterialInterface* MaterialInterface, bool ExpectedIsTwoSided)
+{
+	FInterchangeTestFunctionResult Result;
+	const bool ImportedIsTwoSided = MaterialInterface->IsTwoSided();
+
+	if (ImportedIsTwoSided != ExpectedIsTwoSided)
+	{
+		const TCHAR* ImportedDisplayValue = ImportedIsTwoSided ? TEXT("True") : TEXT("False");
+		const TCHAR* ExpectedDisplayValue = ExpectedIsTwoSided ? TEXT("True") : TEXT("False");
+		Result.AddError(FString::Printf(TEXT("Expected two-sided setting to be %s, imported %s."), ExpectedDisplayValue, ImportedDisplayValue));
+	}
+
+	return Result;
+}
+
+FInterchangeTestFunctionResult UMaterialImportTestFunctions::CheckOpacityMaskClipValue(const UMaterialInterface* MaterialInterface, float ExpectedOpacityMaskClipValue)
+{
+	FInterchangeTestFunctionResult Result;
+	const float ImportedOpacityMaskClipValue = MaterialInterface->GetOpacityMaskClipValue();
+
+	if (!FMath::IsNearlyEqual(ImportedOpacityMaskClipValue, ExpectedOpacityMaskClipValue))
+	{
+		Result.AddError(FString::Printf(TEXT("Expected opacity mask clip value %.2f, imported %.2f."), ExpectedOpacityMaskClipValue, ImportedOpacityMaskClipValue));
+	}
+
+	return Result;
+}
diff --git a/Engine/Plugins/Tests/InterchangeTests/Source/InterchangeTests/Public/ImportTestFunctions/ActorImportTestFunctions.h b/Engine/Plugins/Tests/InterchangeTests/Source/InterchangeTests/Public/ImportTestFunctions/ActorImportTestFunctions.h
--- a/Engine/Plugins/Tests/InterchangeTests/Source/InterchangeTests/Public/ImportTestFunctions/ActorImportTestFunctions.h
+++ b/Engine/Plugins/Tests/InterchangeTests/Source/InterchangeTests/Public/ImportTestFunctions/ActorImportTestFunctions.h
@@ -22,6 +22,22 @@ public:
 	/** Check whether the expected number of actors are imported */
 	UFUNCTION(Exec)
 	static FInterchangeTestFunctionResult CheckImportedActorCount(const TArray<AActor*>& Actors, int32 ExpectedNumberOfImportedActors);
+
+	/** Check whether the expected number of actors with a given class are imported */
+	UFUNCTION(Exec)
+	static FInterchangeTestFunctionResult CheckActorClassCount(const TArray<AActor*>& Actors, TSubclassOf<AActor> Class, int32 ExpectedNumberOfActors);
+
+	/** Check whether the imported actor has the expected class */
+	UFUNCTION(Exec)
+	static FInterchangeTestFunctionResult CheckActorClass(const AActor* Actor, TSubclassOf<AActor> ExpectedClass);
+
+	/** Check whether the generic property (with a given name) in the imported actor has the expected value (or matches with it using regex) */
+	UFUNCTION(Exec)
+	static FInterchangeTestFunctionResult CheckActorPropertyValue(const AActor* Actor, const FString& PropertyName, bool bUseRegexToMatchValue, const FString& ExpectedValue);
+
+	/** Check whether the generic property (with a given name) in the imported actor component (with the given name) has the expected value (or matches with it using regex) */
+	UFUNCTION(Exec)
+	static FInterchangeTestFunctionResult CheckComponentPropertyValue(const AActor* Actor, const FString& ComponentName, const FString& PropertyName, bool bUseRegexToMatchValue, const FString& ExpectedValue);
 };
 
 #if UE_ENABLE_INCLUDE_ORDER_DEPRECATED_IN_5_2
diff --git a/Engine/Plugins/Tests/InterchangeTests/Source/InterchangeTests/Public/ImportTestFunctions/LevelVariantSetsImportTestFunctions.h b/Engine/Plugins/Tests/InterchangeTests/Source/InterchangeTests/Public/ImportTestFunctions/LevelVariantSetsImportTestFunctions.h
new file mode 100644
--- /dev/null
+++ b/Engine/Plugins/Tests/InterchangeTests/Source/InterchangeTests/Public/ImportTestFunctions/LevelVariantSetsImportTestFunctions.h
@@ -0,0 +1,42 @@
+// Copyright Epic Games, Inc. All Rights Reserved.
+
+#pragma once
+
+#include "ImportTestFunctionsBase.h"
+#include "LevelVariantSetsImportTestFunctions.generated.h"
+
+class ULevelVariantSets;
+
+struct FInterchangeTestFunctionResult;
+
+
+UCLASS()
+class INTERCHANGETESTS_API ULevelVariantSetsImportTestFunctions : public UImportTestFunctionsBase
+{
+	GENERATED_BODY()
+
+public:
+
+	// UImportTestFunctionsBase interface
+	virtual UClass* GetAssociatedAssetType() const override;
+
+	/** Check whether the expected number of LevelVariantSets are imported */
+	UFUNCTION(Exec)
+	static FInterchangeTestFunctionResult CheckLevelVariantSetsCount(const TArray<ULevelVariantSets*>& LevelVariantSetsAssets, int32 ExpectedNumberOfLevelVariantSets);
+
+	/** Check whether the imported LevelVariantSets has the expected number of variant sets */
+	UFUNCTION(Exec)
+	static FInterchangeTestFunctionResult CheckVariantSetsCount(ULevelVariantSets* LevelVariantSets, int32 ExpectedNumberOfVariantSets);
+
+	/** Check whether the imported LevelVariantSets has the expected number of variants for the given variant set */
+	UFUNCTION(Exec)
+	static FInterchangeTestFunctionResult CheckVariantsCount(ULevelVariantSets* LevelVariantSets, const FString& VariantSetName, int32 ExpectedNumberOfVariants);
+
+	/** Check whether the imported LevelVariantSets has the expected number of bindings for the given variant in the given set */
+	UFUNCTION(Exec)
+	static FInterchangeTestFunctionResult CheckBindingsCount(ULevelVariantSets* LevelVariantSets, const FString& VariantSetName, const FString& VariantName, int32 ExpectedNumberOfBindings);
+};
+
+#if UE_ENABLE_INCLUDE_ORDER_DEPRECATED_IN_5_2
+#include "InterchangeTestFunction.h"
+#endif
diff --git a/Engine/Plugins/Tests/InterchangeTests/Source/InterchangeTests/Public/ImportTestFunctions/MaterialImportTestFunctions.h b/Engine/Plugins/Tests/InterchangeTests/Source/InterchangeTests/Public/ImportTestFunctions/MaterialImportTestFunctions.h
--- a/Engine/Plugins/Tests/InterchangeTests/Source/InterchangeTests/Public/ImportTestFunctions/MaterialImportTestFunctions.h
+++ b/Engine/Plugins/Tests/InterchangeTests/Source/InterchangeTests/Public/ImportTestFunctions/MaterialImportTestFunctions.h
@@ -2,6 +2,7 @@
 
 #pragma once
 
+#include "Engine/EngineTypes.h"
 #include "ImportTestFunctionsBase.h"
 #include "MaterialImportTestFunctions.generated.h"
 
@@ -27,6 +28,22 @@ public:
 	/** Check whether the expected number of material instances are imported */
 	UFUNCTION(Exec)
 	static FInterchangeTestFunctionResult CheckImportedMaterialInstanceCount(const TArray<UMaterialInterface*>& MaterialInterfaces, int32 ExpectedNumberOfImportedMaterialInstances);
+
+	/** Check whether the imported material has the expected shading model */
+	UFUNCTION(Exec)
+	static FInterchangeTestFunctionResult CheckShadingModel(const UMaterialInterface* MaterialInterface, EMaterialShadingModel ExpectedShadingModel);
+
+	/** Check whether the imported material has the expected blend mode */
+	UFUNCTION(Exec)
+	static FInterchangeTestFunctionResult CheckBlendMode(const UMaterialInterface* MaterialInterface, EBlendMode ExpectedBlendMode);
+
+	/** Check whether the imported material has the expected two-sided setting */
+	UFUNCTION(Exec)
+	static FInterchangeTestFunctionResult CheckIsTwoSided(const UMaterialInterface* MaterialInterface, bool ExpectedIsTwoSided);
+
+	/** Check whether the imported material has the expected opacity mask clip value */
+	UFUNCTION(Exec)
+	static FInterchangeTestFunctionResult CheckOpacityMaskClipValue(const UMaterialInterface* MaterialInterface, float ExpectedOpacityMaskClipValue);
 };
 
 #if UE_ENABLE_INCLUDE_ORDER_DEPRECATED_IN_5_2
