Add ANOTHER 12 necessary test functions for interchange test plans:
- CheckLevelSequenceCount, CheckSequenceLength, CheckSectionCount, CheckSectionInterpolationMode in ULevelSequenceImportTestFunctions
- CheckTextureFilter, CheckTextureAddressX, CheckTextureAddressY, CheckTextureAddressZ, in UTextureImportTestFunctions
- CheckRenderVertexIndexColor, CheckRenderHasVertexColors in UStaticMeshImportTestFunctions
- CheckScalarParameter, CheckVectorParameter in UMaterialImportTestFunctions

diff --git a/Engine/Plugins/Tests/InterchangeTests/Source/InterchangeTests/InterchangeTests.Build.cs b/Engine/Plugins/Tests/InterchangeTests/Source/InterchangeTests/InterchangeTests.Build.cs
--- a/Engine/Plugins/Tests/InterchangeTests/Source/InterchangeTests/InterchangeTests.Build.cs
+++ b/Engine/Plugins/Tests/InterchangeTests/Source/InterchangeTests/InterchangeTests.Build.cs
@@ -27,7 +27,9 @@ namespace UnrealBuildTool.Rules
 					"InterchangeCore",
 					"InterchangeDispatcher",
 					"InterchangeEngine",
+					"LevelSequence",
 					"MeshDescription",
+					"MovieScene",
 					"StaticMeshDescription",
 					"UnrealEd",
 					"VariantManagerContent"
diff --git a/Engine/Plugins/Tests/InterchangeTests/Source/InterchangeTests/Private/ImportTestFunctions/LevelSequenceImportTestFunctions.cpp b/Engine/Plugins/Tests/InterchangeTests/Source/InterchangeTests/Private/ImportTestFunctions/LevelSequenceImportTestFunctions.cpp
new file mode 100644
--- /dev/null
+++ b/Engine/Plugins/Tests/InterchangeTests/Source/InterchangeTests/Private/ImportTestFunctions/LevelSequenceImportTestFunctions.cpp
@@ -0,0 +1,150 @@
+// Copyright Epic Games, Inc. All Rights Reserved.
+
+#include "ImportTestFunctions/LevelSequenceImportTestFunctions.h"
+#include "InterchangeTestFunction.h"
+#include "LevelSequence.h"
+#include "MovieScene.h"
+#include "MovieSceneTimeHelpers.h"
+#include "Channels/MovieSceneChannelProxy.h"
+#include "Channels/MovieSceneFloatChannel.h"
+#include "Channels/MovieSceneDoubleChannel.h"
+
+#include UE_INLINE_GENERATED_CPP_BY_NAME(LevelSequenceImportTestFunctions)
+
+namespace UE::Interchange::Tests::Private
+{
+	bool TryGetSectionInterpolationMode(const UMovieSceneSection* Section, ERichCurveInterpMode& OutInterpMode)
+	{
+		bool HasInterpMode = false;
+		const FMovieSceneChannelProxy& ChannelProxy = Section->GetChannelProxy();
+
+		for (FMovieSceneFloatChannel* Channel : ChannelProxy.GetChannels<FMovieSceneFloatChannel>())
+		{
+			for (const FMovieSceneFloatValue& Value : Channel->GetData().GetValues())
+			{
+				if (!HasInterpMode)
+				{
+					HasInterpMode = true;
+					OutInterpMode = Value.InterpMode;
+				}
+				else if (Value.InterpMode != OutInterpMode)
+				{
+					return false;
+				}
+			}
+		}
+
+		for (FMovieSceneDoubleChannel* Channel : ChannelProxy.GetChannels<FMovieSceneDoubleChannel>())
+		{
+			for (const FMovieSceneDoubleValue& Value : Channel->GetData().GetValues())
+			{
+				if (!HasInterpMode)
+				{
+					HasInterpMode = true;
+					OutInterpMode = Value.InterpMode;
+				}
+				else if (Value.InterpMode != OutInterpMode)
+				{
+					return false;
+				}
+			}
+		}
+
+		return HasInterpMode;
+	}
+}
+
+UClass* ULevelSequenceImportTestFunctions::GetAssociatedAssetType() const
+{
+	return ULevelSequence::StaticClass();
+}
+
+FInterchangeTestFunctionResult ULevelSequenceImportTestFunctions::CheckLevelSequenceCount(const TArray<ULevelSequence*>& LevelSequences, int32 ExpectedNumberOfLevelSequences)
+{
+	FInterchangeTestFunctionResult Result;
+	if (LevelSequences.Num() != ExpectedNumberOfLevelSequences)
+	{
+		Result.AddError(FString::Printf(TEXT("Expected %d level sequences, imported %d."), ExpectedNumberOfLevelSequences, LevelSequences.Num()));
+	}
+
+	return Result;
+}
+
+FInterchangeTestFunctionResult ULevelSequenceImportTestFunctions::CheckSequenceLength(const ULevelSequence* LevelSequence, float ExpectedLevelSequenceLength)
+{
+	FInterchangeTestFunctionResult Result;
+	const TObjectPtr<UMovieScene> MovieScene = LevelSequence->GetMovieScene();
+
+	if (MovieScene == nullptr)
+	{
+		Result.AddError(FString::Printf(TEXT("The imported level sequence doesn't contain a valid movie scene")));
+	}
+	else
+	{
+		const float ImportedLevelSequenceLength = MovieScene->GetTickResolution().AsSeconds(MovieScene->GetPlaybackRange().Size<FFrameNumber>());
+		if (!FMath::IsNearlyEqual(ImportedLevelSequenceLength, ExpectedLevelSequenceLength, UE_KINDA_SMALL_NUMBER))
+		{
+			Result.AddError(FString::Printf(TEXT("Expected level sequence length of %f seconds, imported %f."), ExpectedLevelSequenceLength, ImportedLevelSequenceLength));
+		}
+	}
+
+	return Result;
+}
+
+FInterchangeTestFunctionResult ULevelSequenceImportTestFunctions::CheckSectionCount(const ULevelSequence* LevelSequence, int32 ExpectedNumberOfSections)
+{
+	FInterchangeTestFunctionResult Result;
+	const TObjectPtr<UMovieScene> MovieScene = LevelSequence->GetMovieScene();
+
+	if (MovieScene == nullptr)
+	{
+		Result.AddError(FString::Printf(TEXT("The imported level sequence doesn't contain a valid movie scene")));
+	}
+	else
+	{
+		const int32 ImportedNumberOfSections = MovieScene->GetAllSections().Num();
+		if (ImportedNumberOfSections != ExpectedNumberOfSections)
+		{
+			Result.AddError(FString::Printf(TEXT("Expected %d sections in level sequence, imported %d."), ExpectedNumberOfSections, ImportedNumberOfSections));
+		}
+	}
+
+	return Result;
+}
+
+FInterchangeTestFunctionResult ULevelSequenceImportTestFunctions::CheckSectionInterpolationMode(const ULevelSequence* LevelSequence, int32 SectionIndex, ERichCurveInterpMode ExpectedInterpolationMode)
+{
+	using namespace UE::Interchange::Tests::Private;
+
+	FInterchangeTestFunctionResult Result;
+	const TObjectPtr<UMovieScene> MovieScene = LevelSequence->GetMovieScene();
+
+	if (MovieScene == nullptr)
+	{
+		Result.AddError(FString::Printf(TEXT("The imported level sequence doesn't contain a valid movie scene")));
+	}
+	else
+	{
+		const TArray<UMovieSceneSection*> Sections = MovieScene->GetAllSections();
+		if (!Sections.IsValidIndex(SectionIndex))
+		{
+			Result.AddError(FString::Printf(TEXT("The imported level sequence doesn't contain section %d."), SectionIndex));
+		}
+		else
+		{
+			ERichCurveInterpMode ImportedInterpolationMode;
+			if (!TryGetSectionInterpolationMode(Sections[SectionIndex], ImportedInterpolationMode))
+			{
+				Result.AddError(FString::Printf(TEXT("The imported level sequence doesn't have a consistent interpolation mode for section %d."), SectionIndex));
+			}
+			else if (ImportedInterpolationMode != ExpectedInterpolationMode)
+			{
+				const FString ImportedDisplayValue = UEnum::GetDisplayValueAsText(ImportedInterpolationMode).ToString();
+				const FString ExpectedDisplayValue = UEnum::GetDisplayValueAsText(ExpectedInterpolationMode).ToString();
+				Result.AddError(FString::Printf(TEXT("For section %d, expected interpolation mode %s, imported %s."), SectionIndex, *ExpectedDisplayValue, *ImportedDisplayValue));
+			}
+		}
+	}
+
+	return Result;
+}
diff --git a/Engine/Plugins/Tests/InterchangeTests/Source/InterchangeTests/Private/ImportTestFunctions/MaterialImportTestFunctions.cpp b/Engine/Plugins/Tests/InterchangeTests/Source/InterchangeTests/Private/ImportTestFunctions/MaterialImportTestFunctions.cpp
--- a/Engine/Plugins/Tests/InterchangeTests/Source/InterchangeTests/Private/ImportTestFunctions/MaterialImportTestFunctions.cpp
+++ b/Engine/Plugins/Tests/InterchangeTests/Source/InterchangeTests/Private/ImportTestFunctions/MaterialImportTestFunctions.cpp
@@ -121,10 +121,44 @@ FInterchangeTestFunctionResult UMaterialImportTestFunctions::CheckOpacityMaskCli
 	FInterchangeTestFunctionResult Result;
 	const float ImportedOpacityMaskClipValue = MaterialInterface->GetOpacityMaskClipValue();
 
-	if (!FMath::IsNearlyEqual(ImportedOpacityMaskClipValue, ExpectedOpacityMaskClipValue))
+	if (!FMath::IsNearlyEqual(ImportedOpacityMaskClipValue, ExpectedOpacityMaskClipValue, UE_KINDA_SMALL_NUMBER))
 	{
 		Result.AddError(FString::Printf(TEXT("Expected opacity mask clip value %.2f, imported %.2f."), ExpectedOpacityMaskClipValue, ImportedOpacityMaskClipValue));
 	}
 
 	return Result;
 }
+
+FInterchangeTestFunctionResult UMaterialImportTestFunctions::CheckScalarParameter(const UMaterialInterface* MaterialInterface, const FString& ParameterName, float ExpectedParameterValue)
+{
+	FInterchangeTestFunctionResult Result;
+	float ImportedParameterValue;
+
+	if (!MaterialInterface->GetScalarParameterValue(*ParameterName, ImportedParameterValue))
+	{
+		Result.AddError(FString::Printf(TEXT("The imported material doesn't contain scalar parameter '%s'."), *ParameterName));
+	}
+	else if (!FMath::IsNearlyEqual(ImportedParameterValue, ExpectedParameterValue, UE_KINDA_SMALL_NUMBER))
+	{
+		Result.AddError(FString::Printf(TEXT("For scalar parameter '%s', expected value %f, imported %f."), *ParameterName, ExpectedParameterValue, ImportedParameterValue));
+	}
+
+	return Result;
+}
+
+FInterchangeTestFunctionResult UMaterialImportTestFunctions::CheckVectorParameter(const UMaterialInterface* MaterialInterface, const FString& ParameterName, FLinearColor ExpectedParameterValue)
+{
+	FInterchangeTestFunctionResult Result;
+	FLinearColor ImportedParameterValue;
+
+	if (!MaterialInterface->GetVectorParameterValue(*ParameterName, ImportedParameterValue))
+	{
+		Result.AddError(FString::Printf(TEXT("The imported material doesn't contain vector parameter '%s'."), *ParameterName));
+	}
+	else if (!ImportedParameterValue.Equals(ExpectedParameterValue, UE_KINDA_SMALL_NUMBER))
+	{
+		Result.AddError(FString::Printf(TEXT("For scalar parameter '%s', expected value %s, imported %s."), *ParameterName, *ExpectedParameterValue.ToString(), *ImportedParameterValue.ToString()));
+	}
+
+	return Result;
+}
diff --git a/Engine/Plugins/Tests/InterchangeTests/Source/InterchangeTests/Private/ImportTestFunctions/StaticMeshImportTestFunctions.cpp b/Engine/Plugins/Tests/InterchangeTests/Source/InterchangeTests/Private/ImportTestFunctions/StaticMeshImportTestFunctions.cpp
--- a/Engine/Plugins/Tests/InterchangeTests/Source/InterchangeTests/Private/ImportTestFunctions/StaticMeshImportTestFunctions.cpp
+++ b/Engine/Plugins/Tests/InterchangeTests/Source/InterchangeTests/Private/ImportTestFunctions/StaticMeshImportTestFunctions.cpp
@@ -239,6 +239,64 @@ FInterchangeTestFunctionResult UStaticMeshImportTestFunctions::CheckRenderVertex
 	return Result;
 }
 
+FInterchangeTestFunctionResult UStaticMeshImportTestFunctions::CheckRenderVertexIndexColor(UStaticMesh* Mesh, int32 LodIndex, int32 VertexIndex, const FColor& ExpectedVertexColor)
+{
+	FInterchangeTestFunctionResult Result;
+
+	if (LodIndex < 0 || LodIndex >= Mesh->GetNumLODs())
+	{
+		Result.AddError(FString::Printf(TEXT("The imported mesh doesn't contain LOD index %d."), LodIndex));
+	}
+	else if (Mesh->GetRenderData() && Mesh->GetRenderData()->LODResources.IsValidIndex(LodIndex))
+	{
+		const int32 NumVertices = Mesh->GetRenderData()->LODResources[LodIndex].VertexBuffers.ColorVertexBuffer.GetNumVertices();
+		if (VertexIndex >= NumVertices)
+		{
+			Result.AddError(FString::Printf(TEXT("The imported LOD render color buffer doesn't contain vertex index %d."), VertexIndex));
+		}
+		else
+		{
+			const FColor& ImportedVertexColor = Mesh->GetRenderData()->LODResources[LodIndex].VertexBuffers.ColorVertexBuffer.VertexColor(VertexIndex);
+			if (ImportedVertexColor != ExpectedVertexColor)
+			{
+				Result.AddError(FString::Printf(TEXT("For LOD %d, vertex index %d, expected vertex color %s, imported %s."), LodIndex, VertexIndex, *ExpectedVertexColor.ToString(), *ImportedVertexColor.ToString()));
+			}
+		}
+	}
+	else
+	{
+		Result.AddError(FString::Printf(TEXT("No valid render data For LOD %d."), LodIndex));
+	}
+
+	return Result;
+}
+
+FInterchangeTestFunctionResult UStaticMeshImportTestFunctions::CheckRenderHasVertexColors(UStaticMesh* Mesh, int32 LodIndex, bool bExpectedHasVertexColors)
+{
+	FInterchangeTestFunctionResult Result;
+
+	if (LodIndex < 0 || LodIndex >= Mesh->GetNumLODs())
+	{
+		Result.AddError(FString::Printf(TEXT("The imported mesh doesn't contain LOD index %d."), LodIndex));
+	}
+	else if (Mesh->GetRenderData() && Mesh->GetRenderData()->LODResources.IsValidIndex(LodIndex))
+	{
+		const bool bImportedHasVertexColors = Mesh->GetRenderData()->LODResources[LodIndex].VertexBuffers.ColorVertexBuffer.GetNumVertices() != 0;
+		if (bImportedHasVertexColors != bExpectedHasVertexColors)
+		{
+			const TCHAR* ImportedDisplayValue = bImportedHasVertexColors ? TEXT("True") : TEXT("False");
+			const TCHAR* ExpectedDisplayValue = bExpectedHasVertexColors ? TEXT("True") : TEXT("False");
+			Result.AddError(FString::Printf(TEXT("For LOD %d, expected vertex colors %s, imported %s."), LodIndex, ExpectedDisplayValue, ImportedDisplayValue));
+		}
+	}
+	else
+	{
+		Result.AddError(FString::Printf(TEXT("No valid render data For LOD %d."), LodIndex));
+	}
+
+	return Result;
+}
+
 FInterchangeTestFunctionResult UStaticMeshImportTestFunctions::CheckTriangleCount(UStaticMesh* Mesh, int32 LodIndex, int32 ExpectedTotalNumberOfTriangles)
 {
 	FInterchangeTestFunctionResult Result;
diff --git a/Engine/Plugins/Tests/InterchangeTests/Source/InterchangeTests/Private/ImportTestFunctions/TextureImportTestFunctions.cpp b/Engine/Plugins/Tests/InterchangeTests/Source/InterchangeTests/Private/ImportTestFunctions/TextureImportTestFunctions.cpp
--- a/Engine/Plugins/Tests/InterchangeTests/Source/InterchangeTests/Private/ImportTestFunctions/TextureImportTestFunctions.cpp
+++ b/Engine/Plugins/Tests/InterchangeTests/Source/InterchangeTests/Private/ImportTestFunctions/TextureImportTestFunctions.cpp
@@ -23,3 +23,63 @@ FInterchangeTestFunctionResult UTextureImportTestFunctions::CheckImportedTexture
 
 	return Result;
 }
+
+FInterchangeTestFunctionResult UTextureImportTestFunctions::CheckTextureFilter(const UTexture* Texture, TextureFilter ExpectedTextureFilter)
+{
+	FInterchangeTestFunctionResult Result;
+	const TextureFilter ImportedTextureFilter = Texture->Filter;
+
+	if (ImportedTextureFilter != ExpectedTextureFilter)
+	{
+		const FString ImportedDisplayValue = UEnum::GetDisplayValueAsText(ImportedTextureFilter).ToString();
+		const FString ExpectedDisplayValue = UEnum::GetDisplayValueAsText(ExpectedTextureFilter).ToString();
+		Result.AddError(FString::Printf(TEXT("Expected texture filter mode %s, imported %s."), *ExpectedDisplayValue, *ImportedDisplayValue));
+	}
+
+	return Result;
+}
+
+FInterchangeTestFunctionResult UTextureImportTestFunctions::CheckTextureAddressX(const UTexture* Texture, TextureAddress ExpectedTextureAddressX)
+{
+	FInterchangeTestFunctionResult Result;
+	const TextureAddress ImportedTextureAddressX = Texture->GetTextureAddressX();
+
+	if (ImportedTextureAddressX != ExpectedTextureAddressX)
+	{
+		const FString ImportedDisplayValue = UEnum::GetDisplayValueAsText(ImportedTextureAddressX).ToString();
+		const FString ExpectedDisplayValue = UEnum::GetDisplayValueAsText(ExpectedTextureAddressX).ToString();
+		Result.AddError(FString::Printf(TEXT("For X-axis, expected texture address mode %s, imported %s."), *ExpectedDisplayValue, *ImportedDisplayValue));
+	}
+
+	return Result;
+}
+
+FInterchangeTestFunctionResult UTextureImportTestFunctions::CheckTextureAddressY(const UTexture* Texture, TextureAddress ExpectedTextureAddressY)
+{
+	FInterchangeTestFunctionResult Result;
+	const TextureAddress ImportedTextureAddressY = Texture->GetTextureAddressY();
+
+	if (ImportedTextureAddressY != ExpectedTextureAddressY)
+	{
+		const FString ImportedDisplayValue = UEnum::GetDisplayValueAsText(ImportedTextureAddressY).ToString();
+		const FString ExpectedDisplayValue = UEnum::GetDisplayValueAsText(ExpectedTextureAddressY).ToString();
+		Result.AddError(FString::Printf(TEXT("For Y-axis, expected texture address mode %s, imported %s."), *ExpectedDisplayValue, *ImportedDisplayValue));
+	}
+
+	return Result;
+}
+
+FInterchangeTestFunctionResult UTextureImportTestFunctions::CheckTextureAddressZ(const UTexture* Texture, TextureAddress ExpectedTextureAddressZ)
+{
+	FInterchangeTestFunctionResult Result;
+	const TextureAddress ImportedTextureAddressZ = Texture->GetTextureAddressZ();
+
+	if (ImportedTextureAddressZ != ExpectedTextureAddressZ)
+	{
+		const FString ImportedDisplayValue = UEnum::GetDisplayValueAsText(ImportedTextureAddressZ).ToString();
+		const FString ExpectedDisplayValue = UEnum::GetDisplayValueAsText(ExpectedTextureAddressZ).ToString();
+		Result.AddError(FString::Printf(TEXT("For Z-axis, expected texture address mode %s, imported %s."), *ExpectedDisplayValue, *ImportedDisplayValue));
+	}
+
+	return Result;
+}
diff --git a/Engine/Plugins/Tests/InterchangeTests/Source/InterchangeTests/Public/ImportTestFunctions/LevelSequenceImportTestFunctions.h b/Engine/Plugins/Tests/InterchangeTests/Source/InterchangeTests/Public/ImportTestFunctions/LevelSequenceImportTestFunctions.h
new file mode 100644
--- /dev/null
+++ b/Engine/Plugins/Tests/InterchangeTests/Source/InterchangeTests/Public/ImportTestFunctions/LevelSequenceImportTestFunctions.h
@@ -0,0 +1,42 @@
+// Copyright Epic Games, Inc. All Rights Reserved.
+
+#pragma once
+
+#include "Curves/RealCurve.h"
+#include "ImportTestFunctionsBase.h"
+#include "LevelSequenceImportTestFunctions.generated.h"
+
+class ULevelSequence;
+struct FInterchangeTestFunctionResult;
+
+
+UCLASS()
+class INTERCHANGETESTS_API ULevelSequenceImportTestFunctions : public UImportTestFunctionsBase
+{
+	GENERATED_BODY()
+
+public:
+
+	// UImportTestFunctionsBase interface
+	virtual UClass* GetAssociatedAssetType() const override;
+
+	/** Check whether the expected number of level sequences are imported */
+	UFUNCTION(Exec)
+	static FInterchangeTestFunctionResult CheckLevelSequenceCount(const TArray<ULevelSequence*>& LevelSequences, int32 ExpectedNumberOfLevelSequences);
+
+	/** Check whether the level sequence length (second) is as expected */
+	UFUNCTION(Exec)
+	static FInterchangeTestFunctionResult CheckSequenceLength(const ULevelSequence* LevelSequence, float ExpectedLevelSequenceLength);
+
+	/** Check whether the imported level sequence has the expected number of sections */
+	UFUNCTION(Exec)
+	static FInterchangeTestFunctionResult CheckSectionCount(const ULevelSequence* LevelSequence, int32 ExpectedNumberOfSections);
+
+	/** Check whether the imported level sequence has the expected interpolation mode for the given section */
+	UFUNCTION(Exec)
+	static FInterchangeTestFunctionResult CheckSectionInterpolationMode(const ULevelSequence* LevelSequence, int32 SectionIndex, ERichCurveInterpMode ExpectedInterpolationMode);
+};
+
+#if UE_ENABLE_INCLUDE_ORDER_DEPRECATED_IN_5_2
+#include "InterchangeTestFunction.h"
+#endif
diff --git a/Engine/Plugins/Tests/InterchangeTests/Source/InterchangeTests/Public/ImportTestFunctions/MaterialImportTestFunctions.h b/Engine/Plugins/Tests/InterchangeTests/Source/InterchangeTests/Public/ImportTestFunctions/MaterialImportTestFunctions.h
--- a/Engine/Plugins/Tests/InterchangeTests/Source/InterchangeTests/Public/ImportTestFunctions/MaterialImportTestFunctions.h
+++ b/Engine/Plugins/Tests/InterchangeTests/Source/InterchangeTests/Public/ImportTestFunctions/MaterialImportTestFunctions.h
@@ -44,6 +44,14 @@ public:
 	/** Check whether the imported material has the expected opacity mask clip value */
 	UFUNCTION(Exec)
 	static FInterchangeTestFunctionResult CheckOpacityMaskClipValue(const UMaterialInterface* MaterialInterface, float ExpectedOpacityMaskClipValue);
+
+	/** Check whether the imported material has the expected scalar parameter value */
+	UFUNCTION(Exec)
+	static FInterchangeTestFunctionResult CheckScalarParameter(const UMaterialInterface* MaterialInterface, const FString& ParameterName, float ExpectedParameterValue);
+
+	/** Check whether the imported material has the expected vector parameter value */
+	UFUNCTION(Exec)
+	static FInterchangeTestFunctionResult CheckVectorParameter(const UMaterialInterface* MaterialInterface, const FString& ParameterName, FLinearColor ExpectedParameterValue);
 };
 
 #if UE_ENABLE_INCLUDE_ORDER_DEPRECATED_IN_5_2
diff --git a/Engine/Plugins/Tests/InterchangeTests/Source/InterchangeTests/Public/ImportTestFunctions/StaticMeshImportTestFunctions.h b/Engine/Plugins/Tests/InterchangeTests/Source/InterchangeTests/Public/ImportTestFunctions/StaticMeshImportTestFunctions.h
--- a/Engine/Plugins/Tests/InterchangeTests/Source/InterchangeTests/Public/ImportTestFunctions/StaticMeshImportTestFunctions.h
+++ b/Engine/Plugins/Tests/InterchangeTests/Source/InterchangeTests/Public/ImportTestFunctions/StaticMeshImportTestFunctions.h
@@ -92,6 +92,14 @@ public:
 	UFUNCTION(Exec)
 	static FInterchangeTestFunctionResult CheckRenderVertexIndexNormal(UStaticMesh* Mesh, int32 LodIndex, int32 VertexIndex, const FVector4f& ExpectedVertexNormal);
 
+	/** Check whether the render data vertex index of the given LOD has the expected vertex color */
+	UFUNCTION(Exec)
+	static FInterchangeTestFunctionResult CheckRenderVertexIndexColor(UStaticMesh* Mesh, int32 LodIndex, int32 VertexIndex, const FColor& ExpectedVertexColor);
+
+	/** Check whether the render data contains any vertex colors of the given LOD */
+	UFUNCTION(Exec)
+	static FInterchangeTestFunctionResult CheckRenderHasVertexColors(UStaticMesh* Mesh, int32 LodIndex, bool bExpectedHasVertexColors);
+
 	/** Check whether the mesh description has the expected number of triangles for the given LOD */
 	UFUNCTION(Exec)
 	static FInterchangeTestFunctionResult CheckTriangleCount(UStaticMesh* Mesh, int32 LodIndex, int32 ExpectedTotalNumberOfTriangles);
diff --git a/Engine/Plugins/Tests/InterchangeTests/Source/InterchangeTests/Public/ImportTestFunctions/TextureImportTestFunctions.h b/Engine/Plugins/Tests/InterchangeTests/Source/InterchangeTests/Public/ImportTestFunctions/TextureImportTestFunctions.h
--- a/Engine/Plugins/Tests/InterchangeTests/Source/InterchangeTests/Public/ImportTestFunctions/TextureImportTestFunctions.h
+++ b/Engine/Plugins/Tests/InterchangeTests/Source/InterchangeTests/Public/ImportTestFunctions/TextureImportTestFunctions.h
@@ -2,6 +2,7 @@
 
 #pragma once
 
+#include "Engine/TextureDefines.h"
 #include "ImportTestFunctionsBase.h"
 #include "TextureImportTestFunctions.generated.h"
 
@@ -23,6 +24,22 @@ public:
 	/** Check whether the expected number of textures are imported */
 	UFUNCTION(Exec)
 	static FInterchangeTestFunctionResult CheckImportedTextureCount(const TArray<UTexture*>& Textures, int32 ExpectedNumberOfImportedTextures);
+
+	/** Check whether the imported texture has the expected filtering mode */
+	UFUNCTION(Exec)
+	static FInterchangeTestFunctionResult CheckTextureFilter(const UTexture* Texture, TextureFilter ExpectedTextureFilter);
+
+	/** Check whether the imported texture has the expected addressing mode for X-axis */
+	UFUNCTION(Exec)
+	static FInterchangeTestFunctionResult CheckTextureAddressX(const UTexture* Texture, TextureAddress ExpectedTextureAddressX);
+
+	/** Check whether the imported texture has the expected addressing mode for Y-axis */
+	UFUNCTION(Exec)
+	static FInterchangeTestFunctionResult CheckTextureAddressY(const UTexture* Texture, TextureAddress ExpectedTextureAddressY);
+
+	/** Check whether the imported texture has the expected addressing mode for Z-axis */
+	UFUNCTION(Exec)
+	static FInterchangeTestFunctionResult CheckTextureAddressZ(const UTexture* Texture, TextureAddress ExpectedTextureAddressZ);
 };
 
 #if UE_ENABLE_INCLUDE_ORDER_DEPRECATED_IN_5_2
