[InterchangeTests] Expending available StaticMesh tests with BoundingBox checks.

diff --git a/Engine/Plugins/Tests/InterchangeTests/Source/InterchangeTests/Private/ImportTestFunctions/StaticMeshImportTestFunctions.cpp b/Engine/Plugins/Tests/InterchangeTests/Source/InterchangeTests/Private/ImportTestFunctions/StaticMeshImportTestFunctions.cpp
--- a/Engine/Plugins/Tests/InterchangeTests/Source/InterchangeTests/Private/ImportTestFunctions/StaticMeshImportTestFunctions.cpp
+++ b/Engine/Plugins/Tests/InterchangeTests/Source/InterchangeTests/Private/ImportTestFunctions/StaticMeshImportTestFunctions.cpp
@@ -996,3 +996,69 @@ FInterchangeTestFunctionResult UStaticMeshImportTestFunctions::CheckNaniteSettin
 
 	return Result;
 }
+
+FInterchangeTestFunctionResult UStaticMeshImportTestFunctions::CheckBoundingBoxSizeGreaterThan(UStaticMesh* Mesh, FVector3f GreaterThanSize)
+{
+	FBox BBox = Mesh->GetBoundingBox();
+	FVector	Size = BBox.GetSize();
+	
+	FInterchangeTestFunctionResult Result;
+
+	auto CompareDimension = [](FInterchangeTestFunctionResult& Result, float AxisSize, float ExpectedAxisSize, FString AxisName)
+	{
+		if (AxisSize <= ExpectedAxisSize)
+		{
+			Result.AddError(FString::Printf(TEXT("BoundingBoxSize.%s[%f] is not greater than expected[%f]."), *AxisName, AxisSize, ExpectedAxisSize));
+		}
+	};
+	
+	CompareDimension(Result, Size.X, GreaterThanSize.X, "X");
+	CompareDimension(Result, Size.Y, GreaterThanSize.Y, "Y");
+	CompareDimension(Result, Size.Z, GreaterThanSize.Z, "Z");
+
+	return Result;
+}
+
+FInterchangeTestFunctionResult UStaticMeshImportTestFunctions::CheckBoundingBoxSizeLessThan(UStaticMesh* Mesh, FVector3f LessThanSize)
+{
+	FBox BBox = Mesh->GetBoundingBox();
+	FVector	Size = BBox.GetSize();
+
+	FInterchangeTestFunctionResult Result;
+
+	auto CompareDimension = [](FInterchangeTestFunctionResult& Result, float AxisSize, float ExpectedAxisSize, FString AxisName)
+	{
+		if (AxisSize >= ExpectedAxisSize)
+		{
+			Result.AddError(FString::Printf(TEXT("BoundingBoxSize.%s[%f] is not less than expected[%f]."), *AxisName, AxisSize, ExpectedAxisSize));
+		}
+	};
+
+	CompareDimension(Result, Size.X, LessThanSize.X, "X");
+	CompareDimension(Result, Size.Y, LessThanSize.Y, "Y");
+	CompareDimension(Result, Size.Z, LessThanSize.Z, "Z");
+
+	return Result;
+}
+
+FInterchangeTestFunctionResult UStaticMeshImportTestFunctions::CheckBoundingBoxSize(UStaticMesh* Mesh, FVector3f ExpectedSize)
+{
+	FBox BBox = Mesh->GetBoundingBox();
+	FVector	Size = BBox.GetSize();
+
+	FInterchangeTestFunctionResult Result;
+
+	auto CompareDimension = [](FInterchangeTestFunctionResult& Result, float AxisSize, float ExpectedAxisSize, FString AxisName)
+	{
+		if (!FMath::IsNearlyEqual(AxisSize, ExpectedAxisSize))
+		{
+			Result.AddError(FString::Printf(TEXT("BoundingBoxSize.%s[%f] is not equal to expected[%f]."), *AxisName, AxisSize, ExpectedAxisSize));
+		}
+	};
+
+	CompareDimension(Result, Size.X, ExpectedSize.X, "X");
+	CompareDimension(Result, Size.Y, ExpectedSize.Y, "Y");
+	CompareDimension(Result, Size.Z, ExpectedSize.Z, "Z");
+
+	return Result;
+}
\ No newline at end of file
diff --git a/Engine/Plugins/Tests/InterchangeTests/Source/InterchangeTests/Public/ImportTestFunctions/StaticMeshImportTestFunctions.h b/Engine/Plugins/Tests/InterchangeTests/Source/InterchangeTests/Public/ImportTestFunctions/StaticMeshImportTestFunctions.h
--- a/Engine/Plugins/Tests/InterchangeTests/Source/InterchangeTests/Public/ImportTestFunctions/StaticMeshImportTestFunctions.h
+++ b/Engine/Plugins/Tests/InterchangeTests/Source/InterchangeTests/Public/ImportTestFunctions/StaticMeshImportTestFunctions.h
@@ -166,6 +166,18 @@ public:
 	/** Check whether the static mesh has the expected Nanite settings */
 	UFUNCTION(Exec)
 	static FInterchangeTestFunctionResult CheckNaniteSettings(UStaticMesh* Mesh, const FMeshNaniteSettings& ExpectedNaniteSettings);
+
+	/** Check whether the bounding box size is greater than given value */
+	UFUNCTION(Exec)
+	static FInterchangeTestFunctionResult CheckBoundingBoxSizeGreaterThan(UStaticMesh* Mesh, FVector3f GreaterThanSize);
+
+	/** Check whether the bounding box size is less than given value */
+	UFUNCTION(Exec)
+	static FInterchangeTestFunctionResult CheckBoundingBoxSizeLessThan(UStaticMesh* Mesh, FVector3f LessThanSize);
+
+	/** Check whether the bounding box size is equal to given value */
+	UFUNCTION(Exec)
+	static FInterchangeTestFunctionResult CheckBoundingBoxSize(UStaticMesh* Mesh, FVector3f ExpectedSize);
 };
 
 #if UE_ENABLE_INCLUDE_ORDER_DEPRECATED_IN_5_2
