Adding Test classes in the Interchange Test Plan to test lights import:
-Light
-PointLight
-SpotLight

diff --git a/Engine/Plugins/Tests/InterchangeTests/Source/InterchangeTests/Private/ImportTestFunctions/LightImportTestFunctions.cpp b/Engine/Plugins/Tests/InterchangeTests/Source/InterchangeTests/Private/ImportTestFunctions/LightImportTestFunctions.cpp
new file mode 100644
--- /dev/null
+++ b/Engine/Plugins/Tests/InterchangeTests/Source/InterchangeTests/Private/ImportTestFunctions/LightImportTestFunctions.cpp
@@ -0,0 +1,93 @@
+// Copyright Epic Games, Inc. All Rights Reserved.
+
+#include "ImportTestFunctions/LightImportTestFunctions.h"
+#include "Components/LightComponent.h"
+#include "Engine/Light.h"
+
+UClass* ULightImportTestFunctions::GetAssociatedAssetType() const
+{
+	return ALight::StaticClass();
+}
+
+FInterchangeTestFunctionResult ULightImportTestFunctions::CheckLightPosition(ALight* Light, const FVector& ExpectedLightPosition)
+{
+	FInterchangeTestFunctionResult Result;
+	
+	if(ULightComponent* LightComponent = Light->GetLightComponent())
+	{
+		if(const FVector LightPosition = LightComponent->GetLightPosition(); !ExpectedLightPosition.Equals(LightPosition))
+		{
+			Result.AddError(FString::Printf(TEXT("Expected (%g, %g, %g) light position, imported (%g, %g, %g)."),
+											ExpectedLightPosition.X, ExpectedLightPosition.Y, ExpectedLightPosition.Z,
+											LightPosition.X, LightPosition.Y, LightPosition.Z));
+		}
+	}
+	else
+	{
+		Result.AddError(FString::Printf(TEXT("LightComponent is null for %s."), *Light->GetName()));
+	}
+
+	return Result;
+}
+
+FInterchangeTestFunctionResult ULightImportTestFunctions::CheckLightDirection(ALight* Light, const FVector& ExpectedLightDirection)
+{
+	FInterchangeTestFunctionResult Result;
+
+	if(ULightComponent* LightComponent = Light->GetLightComponent())
+	{
+		if(const FVector LightDirection = LightComponent->GetDirection(); !ExpectedLightDirection.Equals(LightDirection))
+		{
+			Result.AddError(FString::Printf(TEXT("Expected (%g, %g, %g) light position, imported (%g, %g, %g)."),
+											ExpectedLightDirection.X, ExpectedLightDirection.Y, ExpectedLightDirection.Z,
+											LightDirection.X, LightDirection.Y, LightDirection.Z));
+		}
+	}
+	else
+	{
+		Result.AddError(FString::Printf(TEXT("LightComponent is null for %s."), *Light->GetName()));
+	}
+
+	return Result;
+}
+
+FInterchangeTestFunctionResult ULightImportTestFunctions::CheckLightIntensity(ALight* Light, float ExpectedLightIntensity)
+{
+	FInterchangeTestFunctionResult Result;
+
+	if(ULightComponent* LightComponent = Light->GetLightComponent())
+	{
+		if(float LightIntensity = LightComponent->Intensity; !FMath::IsNearlyEqual(LightIntensity, ExpectedLightIntensity))
+		{
+			Result.AddError(FString::Printf(TEXT("Expected %g light intensity, imported %g."), LightIntensity, ExpectedLightIntensity));
+		}
+	}
+	else
+	{
+		Result.AddError(FString::Printf(TEXT("LightComponent is null for %s."), *Light->GetName()));
+	}
+
+	return Result;
+}
+
+FInterchangeTestFunctionResult ULightImportTestFunctions::CheckLightColor(ALight* Light, const FLinearColor& ExpectedLightColor)
+{
+	FInterchangeTestFunctionResult Result;
+
+	if(ULightComponent* LightComponent = Light->GetLightComponent())
+	{
+		if(LightComponent->LightColor != ExpectedLightColor.ToFColor(true))
+		{
+			const FLinearColor LightLinearColor = LightComponent->GetLightColor();
+			Result.AddError(FString::Printf(TEXT("Expected (%g, %g, %g) light color, imported (%g, %g, %g)."),
+											ExpectedLightColor.R, ExpectedLightColor.G, ExpectedLightColor.B,
+											LightLinearColor.R, LightLinearColor.G, LightLinearColor.B));
+		}
+	}
+	else
+	{
+		Result.AddError(FString::Printf(TEXT("LightComponent is null for %s."), *Light->GetName()));
+	}
+
+	return Result;
+}
\ No newline at end of file
diff --git a/Engine/Plugins/Tests/InterchangeTests/Source/InterchangeTests/Private/ImportTestFunctions/PointLightImportTestFunctions.cpp b/Engine/Plugins/Tests/InterchangeTests/Source/InterchangeTests/Private/ImportTestFunctions/PointLightImportTestFunctions.cpp
new file mode 100644
--- /dev/null
+++ b/Engine/Plugins/Tests/InterchangeTests/Source/InterchangeTests/Private/ImportTestFunctions/PointLightImportTestFunctions.cpp
@@ -0,0 +1,29 @@
+// Copyright Epic Games, Inc. All Rights Reserved.
+
+#include "ImportTestFunctions/PointLightImportTestFunctions.h"
+#include "Engine/PointLight.h"
+#include "Components/PointLightComponent.h"
+
+UClass* UPointLightImportTestFunctions::GetAssociatedAssetType() const
+{
+	return APointLight::StaticClass();
+}
+
+FInterchangeTestFunctionResult UPointLightImportTestFunctions::CheckLightFalloffExponent(APointLight* Light, float ExpectedLightFalloff)
+{
+	FInterchangeTestFunctionResult Result;
+
+	if(Light->PointLightComponent)
+	{
+		if(!FMath::IsNearlyEqual(Light->PointLightComponent->LightFalloffExponent, ExpectedLightFalloff))
+		{
+			Result.AddError(FString::Printf(TEXT("Expected %g light falloff exponent, imported %g."), Light->PointLightComponent->LightFalloffExponent, ExpectedLightFalloff));
+		}
+	}
+	else
+	{
+		Result.AddError(FString::Printf(TEXT("LightComponent is null for %s."),	*Light->GetName()));
+	}
+
+	return Result;
+}
\ No newline at end of file
diff --git a/Engine/Plugins/Tests/InterchangeTests/Source/InterchangeTests/Private/ImportTestFunctions/SpotLightImportTestFunctions.cpp b/Engine/Plugins/Tests/InterchangeTests/Source/InterchangeTests/Private/ImportTestFunctions/SpotLightImportTestFunctions.cpp
new file mode 100644
--- /dev/null
+++ b/Engine/Plugins/Tests/InterchangeTests/Source/InterchangeTests/Private/ImportTestFunctions/SpotLightImportTestFunctions.cpp
@@ -0,0 +1,41 @@
+// Copyright Epic Games, Inc. All Rights Reserved.
+
+#include "ImportTestFunctions/SpotLightImportTestFunctions.h"
+#include "Engine/SpotLight.h"
+#include "Components/SpotLightComponent.h"
+
+UClass* USpotLightImportTestFunctions::GetAssociatedAssetType() const
+{
+	return ASpotLight::StaticClass();
+}
+
+FInterchangeTestFunctionResult USpotLightImportTestFunctions::CheckLightInnerConeAngle(ASpotLight* Light, float ExpectedLightInnerConeAngle)
+{
+	FInterchangeTestFunctionResult Result;
+
+	if(Light->SpotLightComponent)
+	{
+		if(!FMath::IsNearlyEqual(Light->SpotLightComponent->InnerConeAngle, ExpectedLightInnerConeAngle))
+		{
+			Result.AddError(FString::Printf(TEXT("Expected %g light inner cone angle, imported %g."), Light->SpotLightComponent->InnerConeAngle, ExpectedLightInnerConeAngle));
+		}
+	}
+	else
+	{
+		Result.AddError(FString::Printf(TEXT("LightComponent is null for %s."), *Light->GetName()));
+	}
+
+	return Result;
+}
+
+FInterchangeTestFunctionResult USpotLightImportTestFunctions::CheckLightOuterConeAngle(ASpotLight* Light, float ExpectedLightOuterConeAngle)
+{
+	FInterchangeTestFunctionResult Result;
+
+	if(Light->SpotLightComponent && !FMath::IsNearlyEqual(Light->SpotLightComponent->OuterConeAngle, ExpectedLightOuterConeAngle))
+	{
+		Result.AddError(FString::Printf(TEXT("Expected %g light outer cone angle, imported %g."), Light->SpotLightComponent->OuterConeAngle, ExpectedLightOuterConeAngle));
+	}
+
+	return Result;
+}
\ No newline at end of file
diff --git a/Engine/Plugins/Tests/InterchangeTests/Source/InterchangeTests/Public/ImportTestFunctions/LightImportTestFunctions.h b/Engine/Plugins/Tests/InterchangeTests/Source/InterchangeTests/Public/ImportTestFunctions/LightImportTestFunctions.h
new file mode 100644
--- /dev/null
+++ b/Engine/Plugins/Tests/InterchangeTests/Source/InterchangeTests/Public/ImportTestFunctions/LightImportTestFunctions.h
@@ -0,0 +1,37 @@
+// Copyright Epic Games, Inc. All Rights Reserved.
+
+#pragma once
+
+#include "ActorImportTestFunctions.h"
+#include "InterchangeTestFunction.h"
+#include "LightImportTestFunctions.generated.h"
+
+struct FInterchangeTestFunctionResult;
+
+
+UCLASS()
+class INTERCHANGETESTS_API ULightImportTestFunctions : public UActorImportTestFunctions
+{
+	GENERATED_BODY()
+
+public:
+
+	// UImportTestFunctionsBase interface
+	virtual UClass* GetAssociatedAssetType() const override;
+
+	/** Check whether the light position is correct*/
+	UFUNCTION(Exec)
+	static FInterchangeTestFunctionResult CheckLightPosition(ALight* Light, const FVector& ExpectedLightPosition);
+
+	/** Check whether the light direction is correct*/
+	UFUNCTION(Exec)
+	static FInterchangeTestFunctionResult CheckLightDirection(ALight* Light, const FVector& ExpectedLightDirection);
+
+	/** Check whether the light intensity is correct*/
+	UFUNCTION(Exec)
+	static FInterchangeTestFunctionResult CheckLightIntensity(ALight* Light, float ExpectedLightIntensity);
+
+	/** Check whether the light color is correct*/
+	UFUNCTION(Exec)
+	static FInterchangeTestFunctionResult CheckLightColor(ALight* Light, const FLinearColor& ExpectedLightColor);
+};
diff --git a/Engine/Plugins/Tests/InterchangeTests/Source/InterchangeTests/Public/ImportTestFunctions/PointLightImportTestFunctions.h b/Engine/Plugins/Tests/InterchangeTests/Source/InterchangeTests/Public/ImportTestFunctions/PointLightImportTestFunctions.h
new file mode 100644
--- /dev/null
+++ b/Engine/Plugins/Tests/InterchangeTests/Source/InterchangeTests/Public/ImportTestFunctions/PointLightImportTestFunctions.h
@@ -0,0 +1,25 @@
+// Copyright Epic Games, Inc. All Rights Reserved.
+
+#pragma once
+
+#include "LightImportTestFunctions.h"
+#include "InterchangeTestFunction.h"
+#include "PointLightImportTestFunctions.generated.h"
+
+struct FInterchangeTestFunctionResult;
+
+
+UCLASS()
+class INTERCHANGETESTS_API UPointLightImportTestFunctions : public ULightImportTestFunctions
+{
+	GENERATED_BODY()
+
+public:
+
+	// UImportTestFunctionsBase interface
+	virtual UClass* GetAssociatedAssetType() const override;
+
+	/** Check whether the light falloff exponent is correct*/
+	UFUNCTION(Exec)
+	static FInterchangeTestFunctionResult CheckLightFalloffExponent(APointLight* Light, float ExpectedLightFalloff);
+};
diff --git a/Engine/Plugins/Tests/InterchangeTests/Source/InterchangeTests/Public/ImportTestFunctions/SpotLightImportTestFunctions.h b/Engine/Plugins/Tests/InterchangeTests/Source/InterchangeTests/Public/ImportTestFunctions/SpotLightImportTestFunctions.h
new file mode 100644
--- /dev/null
+++ b/Engine/Plugins/Tests/InterchangeTests/Source/InterchangeTests/Public/ImportTestFunctions/SpotLightImportTestFunctions.h
@@ -0,0 +1,29 @@
+// Copyright Epic Games, Inc. All Rights Reserved.
+
+#pragma once
+
+#include "PointLightImportTestFunctions.h"
+#include "InterchangeTestFunction.h"
+#include "SpotLightImportTestFunctions.generated.h"
+
+struct FInterchangeTestFunctionResult;
+
+
+UCLASS()
+class INTERCHANGETESTS_API USpotLightImportTestFunctions : public UPointLightImportTestFunctions
+{
+	GENERATED_BODY()
+
+public:
+
+	// UImportTestFunctionsBase interface
+	virtual UClass* GetAssociatedAssetType() const override;
+
+	/** Check whether the light inner cone angle is correct*/
+	UFUNCTION(Exec)
+	static FInterchangeTestFunctionResult CheckLightInnerConeAngle(ASpotLight* Light, float ExpectedLightInnerConeAngle);
+
+	/** Check whether the light outer cone angle is correct*/
+	UFUNCTION(Exec)
+	static FInterchangeTestFunctionResult CheckLightOuterConeAngle(ASpotLight* Light, float ExpectedLightOuterConeAngle);
+};
