Check existence of Warning or Error class for Interchange Test Plan, avoid potential crash if the user doesn't set one.

diff --git a/Engine/Plugins/Tests/InterchangeTests/Source/InterchangeTests/Private/ImportTestFunctions/InterchangeResultImportTestFunctions.cpp b/Engine/Plugins/Tests/InterchangeTests/Source/InterchangeTests/Private/ImportTestFunctions/InterchangeResultImportTestFunctions.cpp
--- a/Engine/Plugins/Tests/InterchangeTests/Source/InterchangeTests/Private/ImportTestFunctions/InterchangeResultImportTestFunctions.cpp
+++ b/Engine/Plugins/Tests/InterchangeTests/Source/InterchangeTests/Private/ImportTestFunctions/InterchangeResultImportTestFunctions.cpp
@@ -33,7 +33,14 @@ FInterchangeTestFunctionResult UInterchangeResultImportTestFunctions::CheckIfErr
 
 	if (!bFoundType)
 	{
-		Result.AddError(FString::Printf(TEXT("No error or warning of class %s was generated."), *ErrorOrWarningClass.Get()->GetName()));
+		if(ErrorOrWarningClass)
+		{
+			Result.AddError(FString::Printf(TEXT("No error or warning of class %s was generated."), *ErrorOrWarningClass.Get()->GetName()));
+		}
+		else
+		{
+			Result.AddError(FString::Printf(TEXT("No error or warning class set.")));
+		}
 	}
 
 	return Result;
