Implement CanRunInEnvironment for FInterchangeImportTest

diff --git a/Engine/Plugins/Tests/InterchangeTests/Source/InterchangeTests/Private/InterchangeImportTest.cpp b/Engine/Plugins/Tests/InterchangeTests/Source/InterchangeTests/Private/InterchangeImportTest.cpp
--- a/Engine/Plugins/Tests/InterchangeTests/Source/InterchangeTests/Private/InterchangeImportTest.cpp
+++ b/Engine/Plugins/Tests/InterchangeTests/Source/InterchangeTests/Private/InterchangeImportTest.cpp
@@ -15,9 +15,38 @@
 #include "ObjectTools.h"
 
 
+class FInterchangeImportTestBase : public FAutomationTestBase
+{
+public:
+	FInterchangeImportTestBase(const FString& InName, const bool bInComplexTask)
+		: FAutomationTestBase(InName, bInComplexTask)
+	{
+	}
+
+	virtual bool CanRunInEnvironment(const FString& TestParams, FString* OutReason, bool* OutWarn) const override;
+};
+
+bool FInterchangeImportTestBase::CanRunInEnvironment(const FString& TestParams, FString* OutReason, bool* OutWarn) const
+{
+	//Make sure interchange worker is available, skip the test execution if unavailable
+	if (!UE::Interchange::FInterchangeDispatcher::IsInterchangeWorkerAvailable())
+	{
+		if (nullptr != OutWarn)
+		{
+			*OutWarn = false;
+		}
 
-IMPLEMENT_COMPLEX_AUTOMATION_TEST(FInterchangeImportTest, "Editor.Interchange.Import", EAutomationTestFlags::EditorContext | EAutomationTestFlags::EngineFilter)
+		if (nullptr != OutReason)
+		{
+			*OutReason = "Interchange worker is not available";
+		}
+
+		return false;
+	}
+	return true;
+}
 
+IMPLEMENT_CUSTOM_COMPLEX_AUTOMATION_TEST(FInterchangeImportTest, FInterchangeImportTestBase, "Editor.Interchange.Import", EAutomationTestFlags::EditorContext | EAutomationTestFlags::EngineFilter)
 
 void FInterchangeImportTest::GetTests(TArray<FString>& OutBeautifiedNames, TArray<FString>& OutTestCommands) const
 {
@@ -27,12 +56,6 @@ void FInterchangeImportTest::GetTests(TArray<FString>& OutBeautifiedNames, TArra
 	return;
 #endif
 
-	//Make sure interchange worker is available, do not run the test if unavailable
-	if (!UE::Interchange::FInterchangeDispatcher::IsInterchangeWorkerAvailable())
-	{
-		return;
-	}
-
 	const FAssetRegistryModule& AssetRegistryModule = FModuleManager::LoadModuleChecked<FAssetRegistryModule>(TEXT("AssetRegistry"));
 	TArray<FAssetData> AllTestPlans;
 	AssetRegistryModule.Get().GetAssetsByClass(UInterchangeImportTestPlan::StaticClass()->GetClassPathName(), AllTestPlans, true);
