Replaced use of AssetTypeActions with AssetDefinition

#rnx

diff --git a/Engine/Plugins/Tests/InterchangeTests/Source/InterchangeTestEditor/InterchangeTestEditor.Build.cs b/Engine/Plugins/Tests/InterchangeTests/Source/InterchangeTestEditor/InterchangeTestEditor.Build.cs
--- a/Engine/Plugins/Tests/InterchangeTests/Source/InterchangeTestEditor/InterchangeTestEditor.Build.cs
+++ b/Engine/Plugins/Tests/InterchangeTests/Source/InterchangeTestEditor/InterchangeTestEditor.Build.cs
@@ -13,12 +13,14 @@ namespace UnrealBuildTool.Rules
 				{
                     "Core",
 					"CoreUObject",
+					"ContentBrowser",
 					"Engine",
 					"EditorStyle",
 					"InputCore",
 					"InterchangeTests",
 					"LevelEditor",
 					"PropertyEditor",
+					"ToolMenus",
 					"UnrealEd"
 				}
 			);
@@ -26,6 +28,7 @@ namespace UnrealBuildTool.Rules
 			PrivateDependencyModuleNames.AddRange(
 				new string[]
 				{
+					"AssetDefinition",
 					"AssetTools",
 					"WorkspaceMenuStructure",
 					"SlateCore",
diff --git a/Engine/Plugins/Tests/InterchangeTests/Source/InterchangeTestEditor/Private/AssetDefinition_InterchangeImportTestPlan.cpp b/Engine/Plugins/Tests/InterchangeTests/Source/InterchangeTestEditor/Private/AssetDefinition_InterchangeImportTestPlan.cpp
new file mode 100644
--- /dev/null
+++ b/Engine/Plugins/Tests/InterchangeTests/Source/InterchangeTestEditor/Private/AssetDefinition_InterchangeImportTestPlan.cpp
@@ -0,0 +1,35 @@
+// Copyright Epic Games, Inc. All Rights Reserved.
+
+#include "AssetDefinition_InterchangeImportTestPlan.h"
+
+#include "InterchangeImportTestStepImport.h"
+#include "InterchangePipelineBase.h"
+#include "InterchangeTestFunction.h"
+#include "ImportTestFunctions/MaterialImportTestFunctions.h"
+
+#include "Algo/AnyOf.h"
+#include "Algo/AllOf.h"
+#include "ContentBrowserMenuContexts.h"
+#include "Hal/FileManager.h"
+#include "JsonObjectConverter.h"
+#include "Misc/FileHelper.h"
+#include "Serialization/LargeMemoryWriter.h"
+#include "Serialization/LargeMemoryReader.h"
+#include "UObject/CoreRedirects.h"
+#include "UObject/UObjectIterator.h"
+
+TConstArrayView<FAssetCategoryPath> UAssetDefinition_InterchangeImportTestPlan::GetAssetCategories() const
+{
+	static const auto Categories = { FAssetCategoryPath(NSLOCTEXT("AssetDefinition_InterchangeImportTestPlan_Category", "Name", "Interchange Import Test Plan")) };
+	return Categories;
+}
+
+EAssetCommandResult UAssetDefinition_InterchangeImportTestPlan::OpenAssets(const FAssetOpenArgs& OpenArgs) const
+{
+	for (UInterchangeImportTestPlan* TestPlan : OpenArgs.LoadObjects<UInterchangeImportTestPlan>())
+	{
+		FSimpleAssetEditor::CreateEditor(EToolkitMode::Standalone, OpenArgs.ToolkitHost, TestPlan);
+	}
+
+	return EAssetCommandResult::Handled;
+}
diff --git a/Engine/Plugins/Tests/InterchangeTests/Source/InterchangeTestEditor/Private/AssetDefinition_InterchangeImportTestPlan.h b/Engine/Plugins/Tests/InterchangeTests/Source/InterchangeTestEditor/Private/AssetDefinition_InterchangeImportTestPlan.h
new file mode 100644
--- /dev/null
+++ b/Engine/Plugins/Tests/InterchangeTests/Source/InterchangeTestEditor/Private/AssetDefinition_InterchangeImportTestPlan.h
@@ -0,0 +1,25 @@
+// Copyright Epic Games, Inc. All Rights Reserved.
+
+#pragma once
+
+#include "AssetDefinitionDefault.h"
+#include "AssetTypeActions_Base.h"
+#include "InterchangeImportTestPlan.h"
+
+#include "AssetDefinition_InterchangeImportTestPlan.generated.h"
+
+UCLASS()
+class UAssetDefinition_InterchangeImportTestPlan : public UAssetDefinitionDefault
+{
+	GENERATED_BODY()
+
+public:
+	// UAssetDefinition Begin
+	virtual FText GetAssetDisplayName() const override { return NSLOCTEXT("AssetTypeActions_InterchangeImportTestPlan", "DisplayName", "Interchange Import Test Plan"); }
+	virtual FLinearColor GetAssetColor() const override { return FColor(125, 0, 200); }
+	virtual TSoftClassPtr<UObject> GetAssetClass() const override { return UInterchangeImportTestPlan::StaticClass(); }
+	virtual EAssetCommandResult OpenAssets(const FAssetOpenArgs& OpenArgs) const override;
+	virtual TConstArrayView<FAssetCategoryPath> GetAssetCategories() const override;
+	virtual bool CanImport() const override { return false; }
+	// UAssetDefinition End
+};
diff --git a/Engine/Plugins/Tests/InterchangeTests/Source/InterchangeTestEditor/Private/AssetTypeActions_InterchangeImportTestPlan.cpp b/Engine/Plugins/Tests/InterchangeTests/Source/InterchangeTestEditor/Private/AssetTypeActions_InterchangeImportTestPlan.cpp
deleted file mode 100644
diff --git a/Engine/Plugins/Tests/InterchangeTests/Source/InterchangeTestEditor/Private/AssetTypeActions_InterchangeImportTestPlan.h b/Engine/Plugins/Tests/InterchangeTests/Source/InterchangeTestEditor/Private/AssetTypeActions_InterchangeImportTestPlan.h
deleted file mode 100644
diff --git a/Engine/Plugins/Tests/InterchangeTests/Source/InterchangeTestEditor/Private/InterchangeTestEditorModule.cpp b/Engine/Plugins/Tests/InterchangeTests/Source/InterchangeTestEditor/Private/InterchangeTestEditorModule.cpp
--- a/Engine/Plugins/Tests/InterchangeTests/Source/InterchangeTestEditor/Private/InterchangeTestEditorModule.cpp
+++ b/Engine/Plugins/Tests/InterchangeTests/Source/InterchangeTestEditor/Private/InterchangeTestEditorModule.cpp
@@ -3,7 +3,7 @@
 #include "InterchangeTestEditorModule.h"
 #include "InterchangeTestFunctionLayout.h"
 #include "InterchangeTestFunction.h"
-#include "AssetTypeActions_InterchangeImportTestPlan.h"
+#include "Modules/ModuleManager.h"
 #include "PropertyEditorModule.h"
 
 
@@ -24,12 +24,6 @@ bool FInterchangeTestEditorModule::IsAvailable()
 
 void FInterchangeTestEditorModule::StartupModule()
 {
-	// Register the InterchangeImportTestPlan asset
-	FAssetToolsModule& AssetToolsModule = FAssetToolsModule::GetModule();
-	IAssetTools& AssetTools = AssetToolsModule.Get();
-	AssetTypeActions = MakeShared<FAssetTypeActions_InterchangeImportTestPlan>();
-	AssetTools.RegisterAssetTypeActions(AssetTypeActions.ToSharedRef());
-
 	// Register the FInterchangeTestFunction struct customization
 	FPropertyEditorModule& PropertyModule = FModuleManager::LoadModuleChecked<FPropertyEditorModule>("PropertyEditor");
 	PropertyModule.RegisterCustomPropertyTypeLayout(FInterchangeTestFunction::StaticStruct()->GetFName(),
@@ -42,10 +36,6 @@ void FInterchangeTestEditorModule::ShutdownModule()
 {
 	if (UObjectInitialized())
 	{
-		FAssetToolsModule& AssetToolsModule = FAssetToolsModule::GetModule();
-		IAssetTools& AssetTools = AssetToolsModule.Get();
-		AssetTools.UnregisterAssetTypeActions(AssetTypeActions.ToSharedRef());
-
 		FPropertyEditorModule& PropertyModule = FModuleManager::LoadModuleChecked<FPropertyEditorModule>("PropertyEditor");
 		PropertyModule.UnregisterCustomPropertyTypeLayout(FInterchangeTestFunction::StaticStruct()->GetFName());
 	}
