Fixed compilation error on LInux

#rnx

diff --git a/Engine/Plugins/Tests/InterchangeTests/Source/InterchangeTestEditor/Private/AssetDefinition_InterchangeImportTestPlan.cpp b/Engine/Plugins/Tests/InterchangeTests/Source/InterchangeTestEditor/Private/AssetDefinition_InterchangeImportTestPlan.cpp
--- a/Engine/Plugins/Tests/InterchangeTests/Source/InterchangeTestEditor/Private/AssetDefinition_InterchangeImportTestPlan.cpp
+++ b/Engine/Plugins/Tests/InterchangeTests/Source/InterchangeTestEditor/Private/AssetDefinition_InterchangeImportTestPlan.cpp
@@ -2,22 +2,6 @@
 
 #include "AssetDefinition_InterchangeImportTestPlan.h"
 
-#include "InterchangeImportTestStepImport.h"
-#include "InterchangePipelineBase.h"
-#include "InterchangeTestFunction.h"
-#include "ImportTestFunctions/MaterialImportTestFunctions.h"
-
-#include "Algo/AnyOf.h"
-#include "Algo/AllOf.h"
-#include "ContentBrowserMenuContexts.h"
-#include "Hal/FileManager.h"
-#include "JsonObjectConverter.h"
-#include "Misc/FileHelper.h"
-#include "Serialization/LargeMemoryWriter.h"
-#include "Serialization/LargeMemoryReader.h"
-#include "UObject/CoreRedirects.h"
-#include "UObject/UObjectIterator.h"
-
 TConstArrayView<FAssetCategoryPath> UAssetDefinition_InterchangeImportTestPlan::GetAssetCategories() const
 {
 	static const auto Categories = { FAssetCategoryPath(NSLOCTEXT("AssetDefinition_InterchangeImportTestPlan_Category", "Name", "Interchange Import Test Plan")) };
