Fixed test error when importing to level

#rnx

diff --git a/Engine/Plugins/Tests/InterchangeTests/Source/InterchangeTests/Private/InterchangeImportTest.cpp b/Engine/Plugins/Tests/InterchangeTests/Source/InterchangeTests/Private/InterchangeImportTest.cpp
--- a/Engine/Plugins/Tests/InterchangeTests/Source/InterchangeTests/Private/InterchangeImportTest.cpp
+++ b/Engine/Plugins/Tests/InterchangeTests/Source/InterchangeTests/Private/InterchangeImportTest.cpp
@@ -319,8 +319,10 @@ bool FInterchangeImportTest::RunTest(const FString& Path)
 		{
 			if (AActor* Actor = Cast<AActor>(ResultObject))
 			{
-				constexpr bool bShouldModifyLevel = false;
-				Actor->GetWorld()->RemoveActor(Actor, bShouldModifyLevel);
+				constexpr bool bShouldModifyLevel = true;
+				Actor->GetWorld()->EditorDestroyActor(Actor, bShouldModifyLevel);
+				// Call UObject::Rename directly on actor to avoid AActor::Rename which unnecessarily unregister and re-register components
+				Actor->UObject::Rename(nullptr, GetTransientPackage(), REN_DontCreateRedirectors | REN_ForceNoResetLoaders);
 			}
 			else
 			{
diff --git a/Engine/Plugins/Tests/InterchangeTests/Source/InterchangeTests/Private/InterchangeImportTestStepImport.cpp b/Engine/Plugins/Tests/InterchangeTests/Source/InterchangeTests/Private/InterchangeImportTestStepImport.cpp
--- a/Engine/Plugins/Tests/InterchangeTests/Source/InterchangeTests/Private/InterchangeImportTestStepImport.cpp
+++ b/Engine/Plugins/Tests/InterchangeTests/Source/InterchangeTests/Private/InterchangeImportTestStepImport.cpp
@@ -1,6 +1,11 @@
 // Copyright Epic Games, Inc. All Rights Reserved.
 
 #include "InterchangeImportTestStepImport.h"
+
+#include "Engine/Engine.h"
+#include "Engine/Level.h"
+#include "Engine/World.h"
+#include "GameFramework/WorldSettings.h"
 #include "HAL/FileManager.h"
 #include "InterchangeImportTestData.h"
 #include "Misc/PackageName.h"
@@ -42,6 +47,25 @@ TTuple<UE::Interchange::FAssetImportResultPtr, UE::Interchange::FSceneImportResu
 
 	if (bImportIntoLevel)
 	{
+		// Create transient world to host data from producer
+		TransientWorld = TStrongObjectPtr<UWorld>(NewObject<UWorld>(GetTransientPackage()));
+		TransientWorld->WorldType = EWorldType::EditorPreview;
+
+		FWorldContext& WorldContext = GEngine->CreateNewWorldContext(TransientWorld->WorldType);
+		WorldContext.SetCurrentWorld(TransientWorld.Get());
+
+		TransientWorld->InitializeNewWorld(UWorld::InitializationValues()
+			.AllowAudioPlayback(false)
+			.CreatePhysicsScene(false)
+			.RequiresHitProxies(false)
+			.CreateNavigation(false)
+			.CreateAISystem(false)
+			.ShouldSimulatePhysics(false)
+			.SetTransactional(false)
+		);
+
+		Params.ImportLevel = TransientWorld->GetCurrentLevel();
+
 		return InterchangeManager.ImportSceneAsync(Data.DestAssetPackagePath, ScopedSourceData.GetSourceData(), Params);
 	}
 	else
@@ -121,6 +145,26 @@ FTestStepResults UInterchangeImportTestStepImport::FinishStep(FInterchangeImport
 	// Run all the tests
 	bool bSuccess = PerformTests(Data, ExecutionInfo);
 
+	if (TransientWorld)
+	{
+		// Only keep assets as result objects since the world and its actors are being destroyed
+		TArray<UObject*> ResultObjects = MoveTemp(Data.ResultObjects);
+		Data.ResultObjects.Reset(ResultObjects.Num());
+
+		for (UObject* Object : ResultObjects)
+		{
+			if (!Object->IsA<AActor>())
+			{
+				Data.ResultObjects.Add(Object);
+			}
+		}
+
+		// Now delete world
+		GEngine->DestroyWorldContext(TransientWorld.Get());
+		TransientWorld->DestroyWorld(true);
+		TransientWorld.Reset();
+	}
+
 	Results.bTestStepSuccess = bSuccess;
 	return Results;
 }
diff --git a/Engine/Plugins/Tests/InterchangeTests/Source/InterchangeTests/Public/InterchangeImportTestStepImport.h b/Engine/Plugins/Tests/InterchangeTests/Source/InterchangeTests/Public/InterchangeImportTestStepImport.h
--- a/Engine/Plugins/Tests/InterchangeTests/Source/InterchangeTests/Public/InterchangeImportTestStepImport.h
+++ b/Engine/Plugins/Tests/InterchangeTests/Source/InterchangeTests/Public/InterchangeImportTestStepImport.h
@@ -3,11 +3,13 @@
 #pragma once
 
 #include "InterchangeImportTestStepBase.h"
-#include "InterchangeImportTestStepImport.generated.h"
 
+#include "UObject/StrongObjectPtr.h"
 
-class UInterchangePipelineBase;
+#include "InterchangeImportTestStepImport.generated.h"
 
+class UInterchangePipelineBase;
+class UWorld;
 
 UCLASS(BlueprintType, Meta = (DisplayName = "Import a file"))
 class INTERCHANGETESTS_API UInterchangeImportTestStepImport : public UInterchangeImportTestStepBase
@@ -42,4 +44,7 @@ public:
 		StartStep(FInterchangeImportTestData& Data) override;
 	virtual FTestStepResults FinishStep(FInterchangeImportTestData& Data, FAutomationTestExecutionInfo& ExecutionInfo) override;
 	virtual FString GetContextString() const override;
+
+private:
+	TStrongObjectPtr<UWorld> TransientWorld;
 };
