[Interchange] EngineTest listing fixes.

#rnx

diff --git a/Engine/Plugins/Tests/InterchangeTests/Source/InterchangeTests/Private/InterchangeImportTest.cpp b/Engine/Plugins/Tests/InterchangeTests/Source/InterchangeTests/Private/InterchangeImportTest.cpp
--- a/Engine/Plugins/Tests/InterchangeTests/Source/InterchangeTests/Private/InterchangeImportTest.cpp
+++ b/Engine/Plugins/Tests/InterchangeTests/Source/InterchangeTests/Private/InterchangeImportTest.cpp
@@ -46,7 +46,9 @@ bool FInterchangeImportTestBase::CanRunInEnvironment(const FString& TestParams,
 	return true;
 }
 
-IMPLEMENT_CUSTOM_COMPLEX_AUTOMATION_TEST(FInterchangeImportTest, FInterchangeImportTestBase, "Editor.Interchange.Import", EAutomationTestFlags::EditorContext | EAutomationTestFlags::EngineFilter)
+IMPLEMENT_CUSTOM_COMPLEX_AUTOMATION_TEST(FInterchangeImportTest, FInterchangeImportTestBase, "Editor.Interchange", EAutomationTestFlags::EditorContext | EAutomationTestFlags::EngineFilter)
+
+const FString InterchangeTestsRootGameFolder = TEXT("/Game/Tests/Interchange/");
 
 void FInterchangeImportTest::GetTests(TArray<FString>& OutBeautifiedNames, TArray<FString>& OutTestCommands) const
 {
@@ -77,7 +79,13 @@ void FInterchangeImportTest::GetTests(TArray<FString>& OutBeautifiedNames, TArra
 	{
 		FString PathAsString = Path.ToString();
 		OutTestCommands.Add(PathAsString);
-		OutBeautifiedNames.Add(FPaths::GetBaseFilename(PathAsString));
+
+		FString BeautifiedName = PathAsString;
+		FPaths::MakePathRelativeTo(BeautifiedName, *InterchangeTestsRootGameFolder);
+		BeautifiedName.ReplaceCharInline(TEXT('/'), TEXT('.'));
+		BeautifiedName.ReplaceCharInline(TEXT('\\'), TEXT('.'));
+
+		OutBeautifiedNames.Add(BeautifiedName);
 	}
 }
 
