UE.EditorAutomation(RunTest=Editor) Win64 - GLTF - Multiple Assets not found fix.

#rnx

#virtualized

diff --git a/Engine/Plugins/Tests/InterchangeTests/Source/InterchangeTests/Private/InterchangeTestFunction.cpp b/Engine/Plugins/Tests/InterchangeTests/Source/InterchangeTests/Private/InterchangeTestFunction.cpp
--- a/Engine/Plugins/Tests/InterchangeTests/Source/InterchangeTests/Private/InterchangeTestFunction.cpp
+++ b/Engine/Plugins/Tests/InterchangeTests/Source/InterchangeTests/Private/InterchangeTestFunction.cpp
@@ -3,6 +3,7 @@
 #include "InterchangeTestFunction.h"
 #include "ImportTestFunctions/ImportTestFunctionsBase.h"
 #include "UObject/StructOnScope.h"
+#include "GameFramework/Actor.h"
 
 #include UE_INLINE_GENERATED_CPP_BY_NAME(InterchangeTestFunction)
 
@@ -216,6 +217,36 @@ FInterchangeTestFunctionResult FInterchangeTestFunction::Invoke(const TArray<UOb
 	TArray<UObject*> AssetsToPass;
 	AssetsToPass.Reserve(AssetsToTest.Num());
 
+	auto NameCheck = [](UObject* Asset, const FString& Name)
+		{
+			if (Asset == nullptr)
+			{
+				return false;
+			}
+
+			if (Name.IsEmpty())
+			{
+				return true;
+			}
+
+			if (Asset->GetClass()->IsChildOf(AActor::StaticClass()))
+			{
+				//For AActors we are interested about the Labels
+				if (Cast<AActor>(Asset)->GetActorLabel() == Name)
+				{
+					return true;
+				}
+			}
+			else
+			{
+				if (Asset->GetName() == Name)
+				{
+					return true;
+				}
+			}
+			return false;
+		};
+
 	UClass* ExpectedClass = AssetClass;
 	for (UObject* Asset : AssetsToTest)
 	{
@@ -223,7 +254,7 @@ FInterchangeTestFunctionResult FInterchangeTestFunction::Invoke(const TArray<UOb
 		{
 			// If the asset is of the expected class, add it if the name matches the supplied optional asset name,
 			// or if we are not checking a specific name.
-			if (OptionalAssetName.IsEmpty() || OptionalAssetName == Asset->GetName())
+			if (NameCheck(Asset, OptionalAssetName))
 			{
 				AssetsToPass.Add(Asset);
 			}
