Adding ability to specify and use pipeline overrides in the Reimport Step of Interchange Test Plan

#rnx

diff --git a/Engine/Plugins/Tests/InterchangeTests/Source/InterchangeTests/Private/InterchangeImportTestStepReimport.cpp b/Engine/Plugins/Tests/InterchangeTests/Source/InterchangeTests/Private/InterchangeImportTestStepReimport.cpp
--- a/Engine/Plugins/Tests/InterchangeTests/Source/InterchangeTests/Private/InterchangeImportTestStepReimport.cpp
+++ b/Engine/Plugins/Tests/InterchangeTests/Source/InterchangeTests/Private/InterchangeImportTestStepReimport.cpp
@@ -49,6 +49,11 @@ TTuple<UE::Interchange::FAssetImportResultPtr, UE::Interchange::FSceneImportResu
 	UE::Interchange::FScopedSourceData ScopedSourceData(SourceFileToReimport.FilePath);
 
 	FImportAssetParameters Params;
+	Params.OverridePipelines.Reserve(PipelineStack.Num());
+	for (TObjectPtr<UInterchangePipelineBase> Pipeline : PipelineStack)
+	{
+		Params.OverridePipelines.Add(Pipeline);
+	}
 	Params.bIsAutomated = true;
 	Params.ReimportAsset = AssetToReimport;
 
diff --git a/Engine/Plugins/Tests/InterchangeTests/Source/InterchangeTests/Public/InterchangeImportTestStepReimport.h b/Engine/Plugins/Tests/InterchangeTests/Source/InterchangeTests/Public/InterchangeImportTestStepReimport.h
--- a/Engine/Plugins/Tests/InterchangeTests/Source/InterchangeTests/Public/InterchangeImportTestStepReimport.h
+++ b/Engine/Plugins/Tests/InterchangeTests/Source/InterchangeTests/Public/InterchangeImportTestStepReimport.h
@@ -7,9 +7,9 @@
 #include "Templates/SubclassOf.h"
 #include "InterchangeImportTestStepReimport.generated.h"
 
+class UInterchangePipelineBase;
 struct FInterchangeTestFunction;
 
-
 UCLASS(BlueprintType, Meta = (DisplayName = "Reimport a file"))
 class INTERCHANGETESTS_API UInterchangeImportTestStepReimport : public UInterchangeImportTestStepBase
 {
@@ -21,6 +21,10 @@ public:
 	UPROPERTY(EditAnywhere, Category = General)
 	FFilePath SourceFileToReimport;
 
+	/** The pipeline stack to use when re-importing (an empty array will use the original import pipelines) */
+	UPROPERTY(EditAnywhere, Instanced, Category = General)
+	TArray<TObjectPtr<UInterchangePipelineBase>> PipelineStack;
+
 	/** The type of the asset to reimport. If only one such asset was imported, this is unambiguous. */
 	UPROPERTY(EditAnywhere, Category = General)
 	TSubclassOf<UObject> AssetTypeToReimport;
