Adding Morph Target related functions for Interchange Test Plan
- Added CheckMorphTargetCount
- Added CheckMorphTargetName

#rnx

diff --git a/Engine/Plugins/Tests/InterchangeTests/Source/InterchangeTests/InterchangeTests.Build.cs b/Engine/Plugins/Tests/InterchangeTests/Source/InterchangeTests/InterchangeTests.Build.cs
--- a/Engine/Plugins/Tests/InterchangeTests/Source/InterchangeTests/InterchangeTests.Build.cs
+++ b/Engine/Plugins/Tests/InterchangeTests/Source/InterchangeTests/InterchangeTests.Build.cs
@@ -31,6 +31,7 @@ namespace UnrealBuildTool.Rules
 					"MeshDescription",
 					"MovieScene",
 					"StaticMeshDescription",
+					"SkeletalMeshDescription",
 					"UnrealEd",
 					"VariantManagerContent"
 				}
diff --git a/Engine/Plugins/Tests/InterchangeTests/Source/InterchangeTests/Private/ImportTestFunctions/SkeletalMeshImportTestFunctions.cpp b/Engine/Plugins/Tests/InterchangeTests/Source/InterchangeTests/Private/ImportTestFunctions/SkeletalMeshImportTestFunctions.cpp
--- a/Engine/Plugins/Tests/InterchangeTests/Source/InterchangeTests/Private/ImportTestFunctions/SkeletalMeshImportTestFunctions.cpp
+++ b/Engine/Plugins/Tests/InterchangeTests/Source/InterchangeTests/Private/ImportTestFunctions/SkeletalMeshImportTestFunctions.cpp
@@ -6,11 +6,13 @@
 #include "Engine/SkeletalMesh.h"
 #include "Engine/SkinnedAssetCommon.h"
 #include "ImportTestFunctions/ImportTestFunctionsBase.h"
-#include "Rendering/SkeletalMeshModel.h"
 #include "InterchangeTestFunction.h"
-#include "Rendering/SkeletalMeshRenderData.h"
 #include "Materials/MaterialInterface.h"
+#include "Rendering/SkeletalMeshModel.h"
+#include "Rendering/SkeletalMeshRenderData.h"
 #include "Rendering/SkeletalMeshLODModel.h"
+#include "Rendering/SkeletalMeshLODImporterData.h"
+#include "SkeletalMeshAttributes.h"
 
 #include UE_INLINE_GENERATED_CPP_BY_NAME(SkeletalMeshImportTestFunctions)
 
@@ -324,7 +326,7 @@ FInterchangeTestFunctionResult USkeletalMeshImportTestFunctions::CheckVertexInde
 			}
 			else
 			{
-				FVector VertexPosition = FVector(Mesh->GetResourceForRendering()->LODRenderData[LodIndex].StaticVertexBuffers.PositionVertexBuffer.VertexPosition(VertexIndex));
+				const FVector VertexPosition = FVector(Mesh->GetResourceForRendering()->LODRenderData[LodIndex].StaticVertexBuffers.PositionVertexBuffer.VertexPosition(VertexIndex));
 				if (!VertexPosition.Equals(ExpectedVertexPosition, UE_DOUBLE_KINDA_SMALL_NUMBER))
 				{
 					Result.AddError(FString::Printf(TEXT("For LOD %d vertex index %d, expected position %s, imported %s."), LodIndex, VertexIndex, *ExpectedVertexPosition.ToString(), *VertexPosition.ToString()));
@@ -375,6 +377,57 @@ FInterchangeTestFunctionResult USkeletalMeshImportTestFunctions::CheckVertexInde
 	return Result;
 }
 
+FInterchangeTestFunctionResult USkeletalMeshImportTestFunctions::CheckMorphTargetCount(USkeletalMesh* Mesh, int32 ExpectedNumberOfMorphTargets)
+{
+	FInterchangeTestFunctionResult Result;
+	constexpr int32 LODIndex0 = 0;
+	FMeshDescription* LOD0MeshDescription = Mesh->GetMeshDescription(LODIndex0);
+	if (LOD0MeshDescription == nullptr)
+	{
+		Result.AddError(FString::Printf(TEXT("could not retrieve LOD 0 for skeletalmesh %s"), *Mesh->GetName()));
+	}
+	else 
+	{
+		FSkeletalMeshConstAttributes SkelMeshAttributes(*LOD0MeshDescription);
+		if (SkelMeshAttributes.GetMorphTargetNames().Num() != ExpectedNumberOfMorphTargets)
+		{
+			Result.AddError(FString::Printf(TEXT("found %d Morph Targets for this skeletalmesh %s - but expected %d Morph Targets"), SkelMeshAttributes.GetMorphTargetNames().Num(), *Mesh->GetName(), ExpectedNumberOfMorphTargets));
+		}
+	}
+
+	return Result;
+}
+
+FInterchangeTestFunctionResult USkeletalMeshImportTestFunctions::CheckMorphTargetName(USkeletalMesh* Mesh, int32 MorphTargetIndex, const FString& ExpectedMorphTargetName)
+{
+	FInterchangeTestFunctionResult Result;
+	constexpr int32 LODIndex0 = 0;
+	FMeshDescription* LOD0MeshDescription = Mesh->GetMeshDescription(LODIndex0);
+	if (LOD0MeshDescription == nullptr)
+	{
+		Result.AddError(FString::Printf(TEXT("could not retrieve LOD 0 for skeletalmesh %s"), *Mesh->GetName()));
+	}
+	else
+	{
+		FSkeletalMeshConstAttributes SkelMeshAttributes(*LOD0MeshDescription);
+		
+		if (MorphTargetIndex < 0 || MorphTargetIndex >= SkelMeshAttributes.GetMorphTargetNames().Num())
+		{
+			Result.AddError(FString::Printf(TEXT("imported skeletalmesh %s doesn't have Morph Target at index %d"), *Mesh->GetName(), MorphTargetIndex));
+		}
+		else
+		{	
+			const FString MorphTargetNameString = SkelMeshAttributes.GetMorphTargetNames()[MorphTargetIndex].ToString();
+			if (!MorphTargetNameString.Equals(ExpectedMorphTargetName))
+			{
+				Result.AddError(FString::Printf(TEXT("morph target at index %d has name %s - but expected %s"), MorphTargetIndex, *MorphTargetNameString, *ExpectedMorphTargetName));
+			}
+		}
+	}
+
+	return Result;
+}
+
 FInterchangeTestFunctionResult USkeletalMeshImportTestFunctions::CheckBoneCount(USkeletalMesh* Mesh, int32 ExpectedNumberOfBones)
 {
 	FInterchangeTestFunctionResult Result;
diff --git a/Engine/Plugins/Tests/InterchangeTests/Source/InterchangeTests/Public/ImportTestFunctions/SkeletalMeshImportTestFunctions.h b/Engine/Plugins/Tests/InterchangeTests/Source/InterchangeTests/Public/ImportTestFunctions/SkeletalMeshImportTestFunctions.h
--- a/Engine/Plugins/Tests/InterchangeTests/Source/InterchangeTests/Public/ImportTestFunctions/SkeletalMeshImportTestFunctions.h
+++ b/Engine/Plugins/Tests/InterchangeTests/Source/InterchangeTests/Public/ImportTestFunctions/SkeletalMeshImportTestFunctions.h
@@ -67,6 +67,14 @@ public:
 	UFUNCTION(Exec)
 	static FInterchangeTestFunctionResult CheckVertexIndexNormal(USkeletalMesh* Mesh, int32 LodIndex, int32 VertexIndex, const FVector& ExpectedVertexNormal);
 
+	/** Check whether the mesh has the expected number of morph targets */
+	UFUNCTION(Exec)
+	static FInterchangeTestFunctionResult CheckMorphTargetCount(USkeletalMesh* Mesh, int32 ExpectedNumberOfMorphTargets);
+
+	/** Check whether the imported morph target name is as expected */
+	UFUNCTION(Exec)
+	static FInterchangeTestFunctionResult CheckMorphTargetName(USkeletalMesh* Mesh, int32 MorphTargetIndex, const FString& ExpectedMorphTargetName);
+
 	/** Check whether the mesh has the expected number of bones */
 	UFUNCTION(Exec)
 	static FInterchangeTestFunctionResult CheckBoneCount(USkeletalMesh* Mesh, int32 ExpectedNumberOfBones);
