[Interchange] Fix a crash when automation tests reload imported packages.

#rnx

diff --git a/Engine/Plugins/Tests/InterchangeTests/Source/InterchangeTests/Private/InterchangeImportTest.cpp b/Engine/Plugins/Tests/InterchangeTests/Source/InterchangeTests/Private/InterchangeImportTest.cpp
--- a/Engine/Plugins/Tests/InterchangeTests/Source/InterchangeTests/Private/InterchangeImportTest.cpp
+++ b/Engine/Plugins/Tests/InterchangeTests/Source/InterchangeTests/Private/InterchangeImportTest.cpp
@@ -250,7 +250,11 @@ bool FInterchangeImportTest::RunTest(const FString& Path)
 						// We don't need the Interchange results object any more, as we've already taken everything from it that we need.
 						// If we don't reset it, it will hold on to the trashed versions of the imported assets during GC.
 						Results[PlanIndex] = {};
-
+						//Make sure interchange result are not garbage collected in case we reload some packages
+						if (Data.InterchangeResults)
+						{
+							Data.InterchangeResults->SetFlags(RF_Standalone);
+						}
 						Data.TestPlan->Steps[StepIndex]->FinishStep(Data, ExecutionInfo);
 
 						if (Data.InterchangeResults)
@@ -270,6 +274,7 @@ bool FInterchangeImportTest::RunTest(const FString& Path)
 									break;
 								}
 							}
+							Data.InterchangeResults->ClearFlags(RF_Standalone);
 						}
 
 						// Finished with the interchange results - null it so it will be GC'd later
diff --git a/Engine/Plugins/Tests/InterchangeTests/Source/InterchangeTests/Private/InterchangeImportTestStepBase.cpp b/Engine/Plugins/Tests/InterchangeTests/Source/InterchangeTests/Private/InterchangeImportTestStepBase.cpp
--- a/Engine/Plugins/Tests/InterchangeTests/Source/InterchangeTests/Private/InterchangeImportTestStepBase.cpp
+++ b/Engine/Plugins/Tests/InterchangeTests/Source/InterchangeTests/Private/InterchangeImportTestStepBase.cpp
@@ -1,9 +1,14 @@
 // Copyright Epic Games, Inc. All Rights Reserved.
 
 #include "InterchangeImportTestStepBase.h"
+
+#include "AssetCompilingManager.h"
+#include "ComponentReregisterContext.h"
 #include "InterchangeImportTestData.h"
 #include "InterchangeTestFunction.h"
+#include "Interfaces/Interface_AsyncCompilation.h"
 #include "Misc/AutomationTest.h"
+#include "UObject/SavePackage.h"
 
 #include UE_INLINE_GENERATED_CPP_BY_NAME(InterchangeImportTestStepBase)
 
@@ -37,3 +42,102 @@ bool UInterchangeImportTestStepBase::PerformTests(FInterchangeImportTestData& Da
 	return bSuccess;
 }
 
+
+void UInterchangeImportTestStepBase::SaveReloadAssets(FInterchangeImportTestData& Data)
+{
+	// First save
+	for (const FAssetData& AssetData : Data.ImportedAssets)
+	{
+		const bool bLoadAsset = false;
+		UObject* AssetObject = AssetData.FastGetAsset(bLoadAsset);
+		if (!AssetObject)
+		{
+			continue;
+		}
+		UPackage* PackageObject = AssetObject->GetPackage();
+		if (!PackageObject)
+		{
+			continue;
+		}
+		AssetObject->MarkPackageDirty();
+		FSavePackageArgs SaveArgs;
+		SaveArgs.TopLevelFlags = RF_Standalone;
+		SaveArgs.SaveFlags = SAVE_NoError;
+		UPackage::SavePackage(PackageObject, AssetObject,
+			*FPackageName::LongPackageNameToFilename(AssetData.PackageName.ToString(), FPackageName::GetAssetPackageExtension()),
+			SaveArgs);
+	}
+
+	// Then rename original objects and their packages, and mark as garbage
+	for (const FAssetData& AssetData : Data.ImportedAssets)
+	{
+		const bool bLoadAsset = false;
+		UObject* AssetObject = AssetData.FastGetAsset(bLoadAsset);
+		if (!AssetObject)
+		{
+			continue;
+		}
+		//Make sure asset compilation is done before renaming them and mark them for garbage collect
+		if (const IInterface_AsyncCompilation* AsyncAsset = Cast<IInterface_AsyncCompilation>(AssetObject))
+		{
+			if (AsyncAsset->IsCompiling())
+			{
+				FAssetCompilingManager::Get().FinishCompilationForObjects({AssetObject});
+			}
+		}
+		UPackage* PackageObject = AssetObject->GetPackage();
+		if (!PackageObject || !ensure(PackageObject == AssetData.GetPackage()))
+		{
+			continue;
+		}
+
+		// Mark all objects in the package as garbage, and remove the standalone flag, so that GC can remove the package later
+		TArray<UObject*> ObjectsInPackage;
+		GetObjectsWithPackage(PackageObject, ObjectsInPackage, true);
+		for (UObject* ObjectInPackage : ObjectsInPackage)
+		{
+			if (!ObjectInPackage)
+			{
+				continue;
+			}
+			ObjectInPackage->ClearFlags(RF_Standalone | RF_Public);
+			ObjectInPackage->MarkAsGarbage();
+		}
+
+		// Renaming the original objects avoids having to do a GC sweep here (this is done at the end of each test step)
+		// Any existing references to them will be retained but irrelevant.
+		// Then the new object can be loaded in their place, as if it were being loaded for the first time.
+		const ERenameFlags RenameFlags = REN_DontCreateRedirectors | REN_NonTransactional | REN_DoNotDirty;
+		PackageObject->Rename(*(PackageObject->GetName() + TEXT("_TRASH")), nullptr, RenameFlags);
+		PackageObject->RemoveFromRoot();
+		PackageObject->MarkAsGarbage();
+
+		// Remove the old version of the asset object from the results
+		Data.ResultObjects.Remove(AssetObject);
+	}
+
+	//Garbage collect objects
+	CollectGarbage(GARBAGE_COLLECTION_KEEPFLAGS);
+
+	// Now reload
+	for (const FAssetData& AssetData : Data.ImportedAssets)
+	{
+		if (!ensure(!AssetData.IsAssetLoaded()))
+		{
+			continue;
+		}
+		UObject* AssetObject = AssetData.GetAsset();
+		if (!AssetObject)
+		{
+			continue;
+		}
+		Data.ResultObjects.Add(AssetObject);
+		if (const IInterface_AsyncCompilation* AsyncAsset = Cast<IInterface_AsyncCompilation>(AssetObject))
+		{
+			if (AsyncAsset->IsCompiling())
+			{
+				FAssetCompilingManager::Get().FinishCompilationForObjects({ AssetObject });
+			}
+		}
+	}
+}
\ No newline at end of file
diff --git a/Engine/Plugins/Tests/InterchangeTests/Source/InterchangeTests/Private/InterchangeImportTestStepImport.cpp b/Engine/Plugins/Tests/InterchangeTests/Source/InterchangeTests/Private/InterchangeImportTestStepImport.cpp
--- a/Engine/Plugins/Tests/InterchangeTests/Source/InterchangeTests/Private/InterchangeImportTestStepImport.cpp
+++ b/Engine/Plugins/Tests/InterchangeTests/Source/InterchangeTests/Private/InterchangeImportTestStepImport.cpp
@@ -9,7 +9,6 @@
 #include "HAL/FileManager.h"
 #include "InterchangeImportTestData.h"
 #include "Misc/PackageName.h"
-#include "UObject/SavePackage.h"
 
 #include UE_INLINE_GENERATED_CPP_BY_NAME(InterchangeImportTestStepImport)
 
@@ -108,62 +107,7 @@ FTestStepResults UInterchangeImportTestStepImport::FinishStep(FInterchangeImport
 	// Incompatible with bImportIntoLevel for now, since we'd need to patch all references to the assets
 	if (bSaveThenReloadImportedAssets && !bImportIntoLevel)
 	{
-		// First save
-		for (const FAssetData& AssetData : Data.ImportedAssets)
-		{
-			const bool bLoadAsset = false;
-			UObject* AssetObject = AssetData.FastGetAsset(bLoadAsset);
-			UPackage* PackageObject = AssetObject->GetPackage();
-			check(PackageObject);
-
-			AssetObject->MarkPackageDirty();
-
-			FSavePackageArgs SaveArgs;
-			SaveArgs.TopLevelFlags = RF_Standalone;
-			SaveArgs.SaveFlags = SAVE_NoError;
-			UPackage::SavePackage(PackageObject, AssetObject,
-				*FPackageName::LongPackageNameToFilename(AssetData.PackageName.ToString(), FPackageName::GetAssetPackageExtension()),
-				SaveArgs);
-		}
-
-		// Then rename original objects and their packages, and mark as garbage
-		for (const FAssetData& AssetData : Data.ImportedAssets)
-		{
-			const bool bLoadAsset = false;
-			UObject* AssetObject = AssetData.FastGetAsset(bLoadAsset);
-			UPackage* PackageObject = AssetObject->GetPackage();
-			check(PackageObject);
-			check(PackageObject == AssetData.GetPackage());
-
-			// Mark all objects in the package as garbage, and remove the standalone flag, so that GC can remove the package later
-			TArray<UObject*> ObjectsInPackage;
-			GetObjectsWithPackage(PackageObject, ObjectsInPackage, true);
-			for (UObject* ObjectInPackage : ObjectsInPackage)
-			{
-				ObjectInPackage->ClearFlags(RF_Standalone | RF_Public);
-				ObjectInPackage->MarkAsGarbage();
-			}
-
-			// Renaming the original objects avoids having to do a GC sweep here (this is done at the end of each test step)
-			// Any existing references to them will be retained but irrelevant.
-			// Then the new object can be loaded in their place, as if it were being loaded for the first time.
-			const ERenameFlags RenameFlags = REN_DontCreateRedirectors | REN_NonTransactional | REN_DoNotDirty;
-			PackageObject->Rename(*(PackageObject->GetName() + TEXT("_TRASH")), nullptr, RenameFlags);
-			PackageObject->RemoveFromRoot();
-			PackageObject->MarkAsGarbage();
-
-			// Remove the old version of the asset object from the results
-			Data.ResultObjects.Remove(AssetObject);
-		}
-
-		// Now reload
-		for (const FAssetData& AssetData : Data.ImportedAssets)
-		{
-			check(!AssetData.IsAssetLoaded());
-			UObject* AssetObject = AssetData.GetAsset();
-			check(AssetObject);
-			Data.ResultObjects.Add(AssetObject);
-		}
+		SaveReloadAssets(Data);
 	}
 
 	// Run all the tests
diff --git a/Engine/Plugins/Tests/InterchangeTests/Source/InterchangeTests/Private/InterchangeImportTestStepReimport.cpp b/Engine/Plugins/Tests/InterchangeTests/Source/InterchangeTests/Private/InterchangeImportTestStepReimport.cpp
--- a/Engine/Plugins/Tests/InterchangeTests/Source/InterchangeTests/Private/InterchangeImportTestStepReimport.cpp
+++ b/Engine/Plugins/Tests/InterchangeTests/Source/InterchangeTests/Private/InterchangeImportTestStepReimport.cpp
@@ -3,7 +3,6 @@
 #include "InterchangeImportTestStepReimport.h"
 #include "InterchangeImportTestData.h"
 #include "Misc/PackageName.h"
-#include "UObject/SavePackage.h"
 
 #include UE_INLINE_GENERATED_CPP_BY_NAME(InterchangeImportTestStepReimport)
 
@@ -70,61 +69,7 @@ FTestStepResults UInterchangeImportTestStepReimport::FinishStep(FInterchangeImpo
 	// If we need to perform a save and fresh reload of everything we imported, do it here
 	if (bSaveThenReloadImportedAssets)
 	{
-		// First save
-		for (const FAssetData& AssetData : Data.ImportedAssets)
-		{
-			const bool bLoadAsset = false;
-			UObject* AssetObject = AssetData.FastGetAsset(bLoadAsset);
-			UPackage* PackageObject = AssetObject->GetPackage();
-			check(PackageObject);
-
-			AssetObject->MarkPackageDirty();
-
-			FSavePackageArgs SaveArgs;
-			SaveArgs.TopLevelFlags = RF_Standalone;
-			SaveArgs.SaveFlags = SAVE_NoError;
-			UPackage::SavePackage(PackageObject, AssetObject,
-				*FPackageName::LongPackageNameToFilename(AssetData.PackageName.ToString(), FPackageName::GetAssetPackageExtension()),
-				SaveArgs);
-		}
-
-		// Then rename original objects and their packages, and mark as garbage
-		for (const FAssetData& AssetData : Data.ImportedAssets)
-		{
-			const bool bLoadAsset = false;
-			UObject* AssetObject = AssetData.FastGetAsset(bLoadAsset);
-			UPackage* PackageObject = AssetObject->GetPackage();
-			check(PackageObject);
-			check(PackageObject == AssetData.GetPackage());
-
-			// Mark all objects in the package as garbage, and remove the standalone flag, so that GC can remove the package later
-			TArray<UObject*> ObjectsInPackage;
-			GetObjectsWithPackage(PackageObject, ObjectsInPackage, true);
-			for (UObject* ObjectInPackage : ObjectsInPackage)
-			{
-				ObjectInPackage->ClearFlags(RF_Standalone | RF_Public);
-				ObjectInPackage->MarkAsGarbage();
-			}
-
-			// Renaming the original objects avoids having to do a GC sweep here.
-			// Any existing references to them will be retained but irrelevant.
-			// Then the new object can be loaded in their place, as if it were being loaded for the first time.
-			const ERenameFlags RenameFlags = REN_DontCreateRedirectors | REN_NonTransactional | REN_DoNotDirty;
-			PackageObject->Rename(*(PackageObject->GetName() + TEXT("_TRASH")), nullptr, RenameFlags);
-			PackageObject->RemoveFromRoot();
-			PackageObject->MarkAsGarbage();
-
-			Data.ResultObjects.Remove(AssetObject);
-		}
-
-		// Now reload
-		for (const FAssetData& AssetData : Data.ImportedAssets)
-		{
-			check(!AssetData.IsAssetLoaded());
-			UObject* AssetObject = AssetData.GetAsset();
-			check(AssetObject);
-			Data.ResultObjects.Add(AssetObject);
-		}
+		SaveReloadAssets(Data);
 	}
 
 	// Run all the tests
diff --git a/Engine/Plugins/Tests/InterchangeTests/Source/InterchangeTests/Public/InterchangeImportTestStepBase.h b/Engine/Plugins/Tests/InterchangeTests/Source/InterchangeTests/Public/InterchangeImportTestStepBase.h
--- a/Engine/Plugins/Tests/InterchangeTests/Source/InterchangeTests/Public/InterchangeImportTestStepBase.h
+++ b/Engine/Plugins/Tests/InterchangeTests/Source/InterchangeTests/Public/InterchangeImportTestStepBase.h
@@ -34,4 +34,6 @@ public:
 	virtual FString GetContextString() const PURE_VIRTUAL(UInterchangeImportTestStepBase::GetContextString, return {}; );
 
 	bool PerformTests(FInterchangeImportTestData& Data, FAutomationTestExecutionInfo& ExecutionInfo);
+protected:
+	void SaveReloadAssets(FInterchangeImportTestData& Data);
 };
