Interchange Test Plan
- Made the FVector Comparison limited to 3 decimals only.

diff --git a/Engine/Plugins/Tests/InterchangeTests/Source/InterchangeTests/Private/ImportTestFunctions/LightImportTestFunctions.cpp b/Engine/Plugins/Tests/InterchangeTests/Source/InterchangeTests/Private/ImportTestFunctions/LightImportTestFunctions.cpp
--- a/Engine/Plugins/Tests/InterchangeTests/Source/InterchangeTests/Private/ImportTestFunctions/LightImportTestFunctions.cpp
+++ b/Engine/Plugins/Tests/InterchangeTests/Source/InterchangeTests/Private/ImportTestFunctions/LightImportTestFunctions.cpp
@@ -3,6 +3,7 @@
 #include "ImportTestFunctions/LightImportTestFunctions.h"
 #include "Components/LightComponent.h"
 #include "Engine/Light.h"
+#include "InterchangeTestsMathUtilities.h"
 
 UClass* ULightImportTestFunctions::GetAssociatedAssetType() const
 {
@@ -15,11 +16,14 @@ FInterchangeTestFunctionResult ULightImportTestFunctions::CheckLightPosition(ALi
 	
 	if(ULightComponent* LightComponent = Light->GetLightComponent())
 	{
-		if(const FVector LightPosition = LightComponent->GetLightPosition(); !ExpectedLightPosition.Equals(LightPosition))
+		using namespace InterchangeTestsMathUtilities;
+		const FVector ExpectedLightPositionRounded = RoundVectorToDecimalPlaces(ExpectedLightPosition);
+		if(const FVector LightPosition = RoundVectorToDecimalPlaces(FVector(LightComponent->GetLightPosition())); 
+			!LightPosition.Equals(ExpectedLightPositionRounded))
 		{
 			Result.AddError(FString::Printf(TEXT("Expected (%g, %g, %g) light position, imported (%g, %g, %g)."),
-											ExpectedLightPosition.X, ExpectedLightPosition.Y, ExpectedLightPosition.Z,
-											LightPosition.X, LightPosition.Y, LightPosition.Z));
+				ExpectedLightPositionRounded.X, ExpectedLightPositionRounded.Y, ExpectedLightPositionRounded.Z,
+				LightPosition.X, LightPosition.Y, LightPosition.Z));
 		}
 	}
 	else
@@ -36,11 +40,13 @@ FInterchangeTestFunctionResult ULightImportTestFunctions::CheckLightDirection(AL
 
 	if(ULightComponent* LightComponent = Light->GetLightComponent())
 	{
-		if(const FVector LightDirection = LightComponent->GetDirection(); !ExpectedLightDirection.Equals(LightDirection))
+		using namespace InterchangeTestsMathUtilities;
+		const FVector ExpectedLightDirectionRounded = RoundVectorToDecimalPlaces(ExpectedLightDirection);
+		if(const FVector LightDirection = RoundVectorToDecimalPlaces(LightComponent->GetDirection()); !LightDirection.Equals(ExpectedLightDirectionRounded))
 		{
-			Result.AddError(FString::Printf(TEXT("Expected (%g, %g, %g) light position, imported (%g, %g, %g)."),
-											ExpectedLightDirection.X, ExpectedLightDirection.Y, ExpectedLightDirection.Z,
-											LightDirection.X, LightDirection.Y, LightDirection.Z));
+			Result.AddError(FString::Printf(TEXT("Expected (%g, %g, %g) light direction, imported (%g, %g, %g)."),
+				ExpectedLightDirectionRounded.X, ExpectedLightDirectionRounded.Y, ExpectedLightDirectionRounded.Z,
+				LightDirection.X, LightDirection.Y, LightDirection.Z));
 		}
 	}
 	else
diff --git a/Engine/Plugins/Tests/InterchangeTests/Source/InterchangeTests/Private/ImportTestFunctions/SkeletalMeshImportTestFunctions.cpp b/Engine/Plugins/Tests/InterchangeTests/Source/InterchangeTests/Private/ImportTestFunctions/SkeletalMeshImportTestFunctions.cpp
--- a/Engine/Plugins/Tests/InterchangeTests/Source/InterchangeTests/Private/ImportTestFunctions/SkeletalMeshImportTestFunctions.cpp
+++ b/Engine/Plugins/Tests/InterchangeTests/Source/InterchangeTests/Private/ImportTestFunctions/SkeletalMeshImportTestFunctions.cpp
@@ -7,6 +7,7 @@
 #include "Engine/SkinnedAssetCommon.h"
 #include "ImportTestFunctions/ImportTestFunctionsBase.h"
 #include "InterchangeTestFunction.h"
+#include "InterchangeTestsMathUtilities.h"
 #include "Materials/MaterialInterface.h"
 #include "Rendering/SkeletalMeshModel.h"
 #include "Rendering/SkeletalMeshRenderData.h"
@@ -326,8 +327,10 @@ FInterchangeTestFunctionResult USkeletalMeshImportTestFunctions::CheckVertexInde
 			}
 			else
 			{
-				const FVector VertexPosition = FVector(Mesh->GetResourceForRendering()->LODRenderData[LodIndex].StaticVertexBuffers.PositionVertexBuffer.VertexPosition(VertexIndex));
-				if (!VertexPosition.Equals(ExpectedVertexPosition, UE_DOUBLE_KINDA_SMALL_NUMBER))
+				using namespace InterchangeTestsMathUtilities;
+				const FVector VertexPosition = RoundVectorToDecimalPlaces(FVector(Mesh->GetResourceForRendering()->LODRenderData[LodIndex].StaticVertexBuffers.PositionVertexBuffer.VertexPosition(VertexIndex)));
+				const FVector ExpectedVertexPositionRounded = RoundVectorToDecimalPlaces(ExpectedVertexPosition);
+				if (!VertexPosition.Equals(ExpectedVertexPositionRounded, UE_DOUBLE_KINDA_SMALL_NUMBER))
 				{
 					Result.AddError(FString::Printf(TEXT("For LOD %d vertex index %d, expected position %s, imported %s."), LodIndex, VertexIndex, *ExpectedVertexPosition.ToString(), *VertexPosition.ToString()));
 				}
@@ -361,8 +364,10 @@ FInterchangeTestFunctionResult USkeletalMeshImportTestFunctions::CheckVertexInde
 			}
 			else
 			{
-				FVector VertexNormal = FVector(Mesh->GetResourceForRendering()->LODRenderData[LodIndex].StaticVertexBuffers.StaticMeshVertexBuffer.VertexTangentZ(VertexIndex));
-				if (!VertexNormal.Equals(ExpectedVertexNormal, UE_DOUBLE_KINDA_SMALL_NUMBER))
+				using namespace InterchangeTestsMathUtilities;
+				const FVector VertexNormal = RoundVectorToDecimalPlaces(FVector(Mesh->GetResourceForRendering()->LODRenderData[LodIndex].StaticVertexBuffers.StaticMeshVertexBuffer.VertexTangentZ(VertexIndex)));
+				const FVector ExpectedVertexNormalRounded = RoundVectorToDecimalPlaces(ExpectedVertexNormal);
+				if (!VertexNormal.Equals(ExpectedVertexNormalRounded, UE_DOUBLE_KINDA_SMALL_NUMBER))
 				{
 					Result.AddError(FString::Printf(TEXT("For LOD %d vertex index %d, expected normal %s, imported %s."), LodIndex, VertexIndex, *ExpectedVertexNormal.ToString(), *VertexNormal.ToString()));
 				}
@@ -471,8 +476,10 @@ FInterchangeTestFunctionResult USkeletalMeshImportTestFunctions::CheckBonePositi
 		}
 		else
 		{
-			FVector BonePosition = Mesh->GetRefSkeleton().GetRefBonePose()[BoneIndex].GetLocation();
-			if (!BonePosition.Equals(ExpectedBonePosition, UE_DOUBLE_KINDA_SMALL_NUMBER))
+			using namespace InterchangeTestsMathUtilities;
+			const FVector BonePosition = RoundVectorToDecimalPlaces(Mesh->GetRefSkeleton().GetRefBonePose()[BoneIndex].GetLocation());
+			const FVector ExpectedBonePositionRounded = RoundVectorToDecimalPlaces(ExpectedBonePosition);
+			if (!BonePosition.Equals(ExpectedBonePositionRounded, UE_DOUBLE_KINDA_SMALL_NUMBER))
 			{
 				Result.AddError(FString::Printf(TEXT("For bone index %d, expected position %s, imported %s."), BoneIndex, *ExpectedBonePosition.ToString(), *BonePosition.ToString()));
 			}
diff --git a/Engine/Plugins/Tests/InterchangeTests/Source/InterchangeTests/Private/ImportTestFunctions/StaticMeshImportTestFunctions.cpp b/Engine/Plugins/Tests/InterchangeTests/Source/InterchangeTests/Private/ImportTestFunctions/StaticMeshImportTestFunctions.cpp
--- a/Engine/Plugins/Tests/InterchangeTests/Source/InterchangeTests/Private/ImportTestFunctions/StaticMeshImportTestFunctions.cpp
+++ b/Engine/Plugins/Tests/InterchangeTests/Source/InterchangeTests/Private/ImportTestFunctions/StaticMeshImportTestFunctions.cpp
@@ -1,15 +1,17 @@
 // Copyright Epic Games, Inc. All Rights Reserved.
 
 #include "ImportTestFunctions/StaticMeshImportTestFunctions.h"
+
 #include "Engine/StaticMesh.h"
 #include "Engine/StaticMeshSocket.h"
 #include "Engine/StaticMeshSourceData.h"
-#include "PhysicsEngine/BodySetup.h"
 #include "ImportTestFunctions/ImportTestFunctionsBase.h"
-#include "StaticMeshAttributes.h"
 #include "InterchangeTestFunction.h"
-#include "StaticMeshResources.h"
+#include "InterchangeTestsMathUtilities.h"
 #include "Materials/MaterialInterface.h"
+#include "PhysicsEngine/BodySetup.h"
+#include "StaticMeshAttributes.h"
+#include "StaticMeshResources.h"
 
 #include UE_INLINE_GENERATED_CPP_BY_NAME(StaticMeshImportTestFunctions)
 
@@ -223,9 +225,10 @@ FInterchangeTestFunctionResult UStaticMeshImportTestFunctions::CheckRenderVertex
 		}
 		else
 		{
-			
-			FVector4f VertexNormal = Mesh->GetRenderData()->LODResources[LodIndex].VertexBuffers.StaticMeshVertexBuffer.VertexTangentZ(VertexIndex);
-			if (!VertexNormal.Equals(ExpectedVertexNormal, UE_KINDA_SMALL_NUMBER))
+			using namespace InterchangeTestsMathUtilities;
+			const FVector4f VertexNormal = RoundVector4ToDecimalPlaces(Mesh->GetRenderData()->LODResources[LodIndex].VertexBuffers.StaticMeshVertexBuffer.VertexTangentZ(VertexIndex));
+			const FVector4f ExpectedVertexNormalRounded = RoundVector4ToDecimalPlaces(ExpectedVertexNormal);
+			if (!VertexNormal.Equals(ExpectedVertexNormalRounded, UE_KINDA_SMALL_NUMBER))
 			{
 				Result.AddError(FString::Printf(TEXT("For LOD %d, vertex index %d, expected vertex normals %s, imported %s."), LodIndex, VertexIndex, *VertexNormal.ToString(), *ExpectedVertexNormal.ToString()));
 			}
@@ -641,9 +644,11 @@ FInterchangeTestFunctionResult UStaticMeshImportTestFunctions::CheckVertexIndexP
 		}
 		else
 		{
-			const FVector VertexPosition = FVector(Attributes.GetVertexPositions()[VertexIndex]);
-			if (!VertexPosition.Equals(ExpectedVertexPosition, UE_DOUBLE_KINDA_SMALL_NUMBER))
-			{
+			using namespace InterchangeTestsMathUtilities;
+			const FVector VertexPosition = RoundVectorToDecimalPlaces(FVector(Attributes.GetVertexPositions()[VertexIndex]));
+			const FVector ExpectedVertexPositionRounded = RoundVectorToDecimalPlaces(ExpectedVertexPosition);
+			if (!VertexPosition.Equals(ExpectedVertexPositionRounded, UE_DOUBLE_KINDA_SMALL_NUMBER))
+			{				
 				Result.AddError(FString::Printf(TEXT("For LOD %d vertex index %d, expected position %s, imported %s."), LodIndex, VertexIndex, *ExpectedVertexPosition.ToString(), *VertexPosition.ToString()));
 			}
 		}
@@ -749,8 +754,11 @@ FInterchangeTestFunctionResult UStaticMeshImportTestFunctions::CheckSocketLocati
 		Result.AddError(FString::Printf(TEXT("The imported mesh doesn't contain %d sockets."), SocketIndex));
 	}
 
-	FVector SocketLocation = Mesh->Sockets[SocketIndex]->RelativeLocation;
-	if (SocketLocation.Equals(ExpectedSocketLocation, UE_DOUBLE_KINDA_SMALL_NUMBER))
+	using namespace InterchangeTestsMathUtilities;
+
+	const FVector SocketLocation = RoundVectorToDecimalPlaces(Mesh->Sockets[SocketIndex]->RelativeLocation);
+	const FVector ExpectedSocketLocationRounded = RoundVectorToDecimalPlaces(ExpectedSocketLocation);
+	if (SocketLocation.Equals(ExpectedSocketLocationRounded, UE_DOUBLE_KINDA_SMALL_NUMBER))
 	{
 		Result.AddError(FString::Printf(TEXT("Expected socket location %s, imported %s"), *ExpectedSocketLocation.ToString(), *SocketLocation.ToString()));
 	}
@@ -905,8 +913,9 @@ FInterchangeTestFunctionResult UStaticMeshImportTestFunctions::CheckAgainstGroun
 			{
 				for (int32 VertexIndex = 0; VertexIndex < VertexCount; VertexIndex++)
 				{
-					const FVector3f& V0 = VertexPositions[VertexIndex];
-					const FVector3f& V1 = GroundTruthVertexPositions[VertexIndex];
+					using namespace InterchangeTestsMathUtilities;
+					const FVector3f V0 = RoundVectorToDecimalPlaces(VertexPositions[VertexIndex]);
+					const FVector3f V1 = RoundVectorToDecimalPlaces(GroundTruthVertexPositions[VertexIndex]);
 
 					if (!V0.Equals(V1))
 					{
@@ -941,8 +950,9 @@ FInterchangeTestFunctionResult UStaticMeshImportTestFunctions::CheckAgainstGroun
 			{
 				for (int32 NormalIndex = 0; NormalIndex < NormalsCount; NormalIndex++)
 				{
-					const FVector3f& N0 = VertexInstanceNormals[NormalIndex];
-					const FVector3f& N1 = GroundTruthVertexInstanceNormals[NormalIndex];
+					using namespace InterchangeTestsMathUtilities;
+					const FVector3f N0 = RoundVectorToDecimalPlaces(VertexInstanceNormals[NormalIndex]);
+					const FVector3f N1 = RoundVectorToDecimalPlaces(GroundTruthVertexInstanceNormals[NormalIndex]);
 
 					if (!N0.Equals(N1))
 					{
@@ -997,68 +1007,74 @@ FInterchangeTestFunctionResult UStaticMeshImportTestFunctions::CheckNaniteSettin
 	return Result;
 }
 
-FInterchangeTestFunctionResult UStaticMeshImportTestFunctions::CheckBoundingBoxSizeGreaterThan(UStaticMesh* Mesh, FVector3f GreaterThanSize)
+FInterchangeTestFunctionResult UStaticMeshImportTestFunctions::CheckBoundingBoxSizeGreaterThan(UStaticMesh* Mesh, const FVector& GreaterThanSize)
 {
+	using namespace InterchangeTestsMathUtilities;
 	FBox BBox = Mesh->GetBoundingBox();
-	FVector	Size = BBox.GetSize();
+	const FVector Size = RoundVectorToDecimalPlaces(BBox.GetSize());
+	const FVector GreaterThanSizeRounded = RoundVectorToDecimalPlaces(GreaterThanSize);
 	
 	FInterchangeTestFunctionResult Result;
 
-	auto CompareDimension = [](FInterchangeTestFunctionResult& Result, float AxisSize, float ExpectedAxisSize, FString AxisName)
+	auto CompareDimension = [](FInterchangeTestFunctionResult& Result, double AxisSize, double ExpectedAxisSize, FString AxisName)
 	{
 		if (AxisSize <= ExpectedAxisSize)
 		{
-			Result.AddError(FString::Printf(TEXT("BoundingBoxSize.%s[%f] is not greater than expected[%f]."), *AxisName, AxisSize, ExpectedAxisSize));
+			Result.AddError(FString::Printf(TEXT("BoundingBoxSize.%s[%lf] is not greater than expected[%lf]."), *AxisName, AxisSize, ExpectedAxisSize));
 		}
 	};
 	
-	CompareDimension(Result, Size.X, GreaterThanSize.X, "X");
-	CompareDimension(Result, Size.Y, GreaterThanSize.Y, "Y");
-	CompareDimension(Result, Size.Z, GreaterThanSize.Z, "Z");
+	CompareDimension(Result, Size.X, GreaterThanSizeRounded.X, "X");
+	CompareDimension(Result, Size.Y, GreaterThanSizeRounded.Y, "Y");
+	CompareDimension(Result, Size.Z, GreaterThanSizeRounded.Z, "Z");
 
 	return Result;
 }
 
-FInterchangeTestFunctionResult UStaticMeshImportTestFunctions::CheckBoundingBoxSizeLessThan(UStaticMesh* Mesh, FVector3f LessThanSize)
+FInterchangeTestFunctionResult UStaticMeshImportTestFunctions::CheckBoundingBoxSizeLessThan(UStaticMesh* Mesh, const FVector& LessThanSize)
 {
+	using namespace InterchangeTestsMathUtilities;
 	FBox BBox = Mesh->GetBoundingBox();
-	FVector	Size = BBox.GetSize();
+	const FVector Size = RoundVectorToDecimalPlaces(BBox.GetSize());
+	const FVector LessThanSizeRounded = RoundVectorToDecimalPlaces(LessThanSize);
 
 	FInterchangeTestFunctionResult Result;
 
-	auto CompareDimension = [](FInterchangeTestFunctionResult& Result, float AxisSize, float ExpectedAxisSize, FString AxisName)
+	auto CompareDimension = [](FInterchangeTestFunctionResult& Result, double AxisSize, double ExpectedAxisSize, FString AxisName)
 	{
 		if (AxisSize >= ExpectedAxisSize)
 		{
-			Result.AddError(FString::Printf(TEXT("BoundingBoxSize.%s[%f] is not less than expected[%f]."), *AxisName, AxisSize, ExpectedAxisSize));
+			Result.AddError(FString::Printf(TEXT("BoundingBoxSize.%s[%lf] is not less than expected[%lf]."), *AxisName, AxisSize, ExpectedAxisSize));
 		}
 	};
 
-	CompareDimension(Result, Size.X, LessThanSize.X, "X");
-	CompareDimension(Result, Size.Y, LessThanSize.Y, "Y");
-	CompareDimension(Result, Size.Z, LessThanSize.Z, "Z");
+	CompareDimension(Result, Size.X, LessThanSizeRounded.X, "X");
+	CompareDimension(Result, Size.Y, LessThanSizeRounded.Y, "Y");
+	CompareDimension(Result, Size.Z, LessThanSizeRounded.Z, "Z");
 
 	return Result;
 }
 
-FInterchangeTestFunctionResult UStaticMeshImportTestFunctions::CheckBoundingBoxSize(UStaticMesh* Mesh, FVector3f ExpectedSize)
+FInterchangeTestFunctionResult UStaticMeshImportTestFunctions::CheckBoundingBoxSize(UStaticMesh* Mesh, const FVector& ExpectedSize)
 {
+	using namespace InterchangeTestsMathUtilities;
 	FBox BBox = Mesh->GetBoundingBox();
-	FVector	Size = BBox.GetSize();
+	const FVector Size = RoundVectorToDecimalPlaces(BBox.GetSize());
+	const FVector ExpectedSizeRounded = RoundVectorToDecimalPlaces(ExpectedSize);
 
 	FInterchangeTestFunctionResult Result;
 
-	auto CompareDimension = [](FInterchangeTestFunctionResult& Result, float AxisSize, float ExpectedAxisSize, FString AxisName)
+	auto CompareDimension = [](FInterchangeTestFunctionResult& Result, double AxisSize, double ExpectedAxisSize, FString AxisName)
 	{
 		if (!FMath::IsNearlyEqual(AxisSize, ExpectedAxisSize))
 		{
-			Result.AddError(FString::Printf(TEXT("BoundingBoxSize.%s[%f] is not equal to expected[%f]."), *AxisName, AxisSize, ExpectedAxisSize));
+			Result.AddError(FString::Printf(TEXT("BoundingBoxSize.%s[%lf] is not equal to expected[%lf]."), *AxisName, AxisSize, ExpectedAxisSize));
 		}
 	};
 
-	CompareDimension(Result, Size.X, ExpectedSize.X, "X");
-	CompareDimension(Result, Size.Y, ExpectedSize.Y, "Y");
-	CompareDimension(Result, Size.Z, ExpectedSize.Z, "Z");
+	CompareDimension(Result, Size.X, ExpectedSizeRounded.X, "X");
+	CompareDimension(Result, Size.Y, ExpectedSizeRounded.Y, "Y");
+	CompareDimension(Result, Size.Z, ExpectedSizeRounded.Z, "Z");
 
 	return Result;
 }
\ No newline at end of file
diff --git a/Engine/Plugins/Tests/InterchangeTests/Source/InterchangeTests/Private/InterchangeTestsMathUtilities.h b/Engine/Plugins/Tests/InterchangeTests/Source/InterchangeTests/Private/InterchangeTestsMathUtilities.h
new file mode 100644
--- /dev/null
+++ b/Engine/Plugins/Tests/InterchangeTests/Source/InterchangeTests/Private/InterchangeTestsMathUtilities.h
@@ -0,0 +1,42 @@
+// Copyright Epic Games, Inc. All Rights Reserved.
+#pragma once
+#include "Math/UnrealMath.h"
+
+namespace InterchangeTestsMathUtilities
+{
+	template<typename Type, uint32 DecimalPlaces>
+	constexpr Type ComputeRoundingFactor()
+	{
+		if constexpr(DecimalPlaces == 0)
+		{
+			return Type(1);
+		}
+		else 
+		{
+			return Type(10) * ComputeRoundingFactor<Type, DecimalPlaces - 1>();
+		}
+	}
+
+	template<typename Type, uint32 DecimalPlaces = 3>
+	UE::Math::TVector<Type> RoundVectorToDecimalPlaces(const UE::Math::TVector<Type>& Vector)
+	{
+		Type Factor = ComputeRoundingFactor<Type, DecimalPlaces>();
+		UE::Math::TVector<Type> Result(Vector);
+		Result.X = (Type)(FMath::RoundToDouble(Result.X * Factor) / Factor);
+		Result.Y = (Type)(FMath::RoundToDouble(Result.Y * Factor) / Factor);
+		Result.Z = (Type)(FMath::RoundToDouble(Result.Z * Factor) / Factor);
+		return Result;
+	}
+
+	template<typename Type, uint32 DecimalPlaces = 3>
+	UE::Math::TVector4<Type> RoundVector4ToDecimalPlaces(const UE::Math::TVector4<Type>& Vector)
+	{
+		Type Factor = ComputeRoundingFactor<Type, DecimalPlaces>();
+		UE::Math::TVector4<Type> Result(Vector);
+		Result.X = (Type)(FMath::RoundToDouble(Result.X * Factor) / Factor);
+		Result.Y = (Type)(FMath::RoundToDouble(Result.Y * Factor) / Factor);
+		Result.Z = (Type)(FMath::RoundToDouble(Result.Z * Factor) / Factor);
+		Result.W = (Type)(FMath::RoundToDouble(Result.W * Factor) / Factor);
+		return Result;
+	}
+}
\ No newline at end of file
diff --git a/Engine/Plugins/Tests/InterchangeTests/Source/InterchangeTests/Public/ImportTestFunctions/StaticMeshImportTestFunctions.h b/Engine/Plugins/Tests/InterchangeTests/Source/InterchangeTests/Public/ImportTestFunctions/StaticMeshImportTestFunctions.h
--- a/Engine/Plugins/Tests/InterchangeTests/Source/InterchangeTests/Public/ImportTestFunctions/StaticMeshImportTestFunctions.h
+++ b/Engine/Plugins/Tests/InterchangeTests/Source/InterchangeTests/Public/ImportTestFunctions/StaticMeshImportTestFunctions.h
@@ -169,13 +169,13 @@ public:
 
 	/** Check whether the bounding box size is greater than given value */
 	UFUNCTION(Exec)
-	static FInterchangeTestFunctionResult CheckBoundingBoxSizeGreaterThan(UStaticMesh* Mesh, FVector3f GreaterThanSize);
+	static FInterchangeTestFunctionResult CheckBoundingBoxSizeGreaterThan(UStaticMesh* Mesh, const FVector& GreaterThanSize);
 
 	/** Check whether the bounding box size is less than given value */
 	UFUNCTION(Exec)
-	static FInterchangeTestFunctionResult CheckBoundingBoxSizeLessThan(UStaticMesh* Mesh, FVector3f LessThanSize);
+	static FInterchangeTestFunctionResult CheckBoundingBoxSizeLessThan(UStaticMesh* Mesh, const FVector& LessThanSize);
 
 	/** Check whether the bounding box size is equal to given value */
 	UFUNCTION(Exec)
-	static FInterchangeTestFunctionResult CheckBoundingBoxSize(UStaticMesh* Mesh, FVector3f ExpectedSize);
+	static FInterchangeTestFunctionResult CheckBoundingBoxSize(UStaticMesh* Mesh, const FVector& ExpectedSize);
 };
