[Interchange] Implement ability for tester to select and modify pipelines from stack presets
- The testers need to open the Import Dialog from the Test Plan asset via a button to change the pipeline settings.
- Current override pipelines could also be edited through import dialog.

#rnx

#virtualized

diff --git a/Engine/Plugins/Interchange/Editor/Source/Pipelines/Private/InterchangePipelineConfigurationGeneric.cpp b/Engine/Plugins/Interchange/Editor/Source/Pipelines/Private/InterchangePipelineConfigurationGeneric.cpp
--- a/Engine/Plugins/Interchange/Editor/Source/Pipelines/Private/InterchangePipelineConfigurationGeneric.cpp
+++ b/Engine/Plugins/Interchange/Editor/Source/Pipelines/Private/InterchangePipelineConfigurationGeneric.cpp
@@ -92,6 +92,7 @@ EInterchangePipelineConfigurationDialogResult UInterchangePipelineConfigurationG
 		.SourceData(SourceData)
 		.bSceneImport(false)
 		.bReimport(false)
+		.bTestConfigDialog(false)
 		.PipelineStacks(PipelineStacks)
 		.OutPipelines(&OutPipelines)
 		.BaseNodeContainer(BaseNodeContainer)
@@ -149,6 +150,7 @@ EInterchangePipelineConfigurationDialogResult UInterchangePipelineConfigurationG
 		.SourceData(SourceData)
 		.bSceneImport(true)
 		.bReimport(false)
+		.bTestConfigDialog(false)
 		.PipelineStacks(PipelineStacks)
 		.OutPipelines(&OutPipelines)
 		.BaseNodeContainer(BaseNodeContainer)
@@ -209,6 +211,7 @@ EInterchangePipelineConfigurationDialogResult UInterchangePipelineConfigurationG
 		.SourceData(SourceData)
 		.bSceneImport(bSceneImport)
 		.bReimport(true)
+		.bTestConfigDialog(false)
 		.PipelineStacks(PipelineStacks)
 		.OutPipelines(&OutPipelines)
 		.BaseNodeContainer(BaseNodeContainer)
@@ -235,3 +238,54 @@ EInterchangePipelineConfigurationDialogResult UInterchangePipelineConfigurationG
 
 	return EInterchangePipelineConfigurationDialogResult::Import;
 }
+
+EInterchangePipelineConfigurationDialogResult UInterchangePipelineConfigurationGeneric::ShowTestPlanConfigurationDialog(TArray<FInterchangeStackInfo>& PipelineStacks, TArray<UInterchangePipelineBase*>& OutPipelines, TWeakObjectPtr<UInterchangeSourceData> SourceData, TWeakObjectPtr<UInterchangeTranslatorBase> Translator, TWeakObjectPtr<UInterchangeBaseNodeContainer> BaseNodeContainer, TWeakObjectPtr<UObject> ReimportAsset, bool bSceneImport, bool bReimport)
+{
+	using namespace UE::Interchange::Private;
+	//Create and show the graph inspector UI dialog
+	TSharedPtr<SWindow> ParentWindow;
+	if (IMainFrameModule* MainFrame = FModuleManager::LoadModulePtr<IMainFrameModule>("MainFrame"))
+	{
+		ParentWindow = MainFrame->GetParentWindow();
+	}
+
+	FString WindowDialogUid = TEXT("TestPlanConfigurationDialog");
+	FVector2D ClientSize = DefaultClientSize();
+	RestoreClientSize(ClientSize, WindowDialogUid);
+
+	TSharedRef<SWindow> Window = SNew(SWindow)
+		.ClientSize(ClientSize)
+		.MinWidth(GetMinSizeX())
+		.MinHeight(GetMinSizeY())
+		.Title(NSLOCTEXT("Interchange", "PipelineConfigurationGenericTitlePipelineConfiguration", "Pipeline Configuration"));
+	TSharedPtr<SInterchangePipelineConfigurationDialog> InterchangePipelineConfigurationDialog;
+
+	Window->SetContent
+	(
+		SAssignNew(InterchangePipelineConfigurationDialog, SInterchangePipelineConfigurationDialog)
+		.OwnerWindow(Window)
+		.SourceData(SourceData)
+		.bSceneImport(bSceneImport)
+		.bReimport(bReimport)
+		.bTestConfigDialog(true)
+		.PipelineStacks(PipelineStacks)
+		.OutPipelines(&OutPipelines)
+		.BaseNodeContainer(BaseNodeContainer)
+		.ReimportObject(ReimportAsset)
+		.Translator(Translator)
+	);
+
+	Window->SetOnWindowClosed(FOnWindowClosed::CreateLambda([&WindowDialogUid](const TSharedRef<SWindow>& WindowMoved)
+		{
+			BackupClientSize(&WindowMoved.Get(), WindowDialogUid);
+		}));
+
+	FSlateApplication::Get().AddModalWindow(Window, ParentWindow, false);
+
+	if (InterchangePipelineConfigurationDialog->IsCanceled())
+	{
+		return EInterchangePipelineConfigurationDialogResult::Cancel;
+	}
+
+	return EInterchangePipelineConfigurationDialogResult::SaveConfig;
+}
diff --git a/Engine/Plugins/Interchange/Editor/Source/Pipelines/Private/SInterchangePipelineConfigurationDialog.cpp b/Engine/Plugins/Interchange/Editor/Source/Pipelines/Private/SInterchangePipelineConfigurationDialog.cpp
--- a/Engine/Plugins/Interchange/Editor/Source/Pipelines/Private/SInterchangePipelineConfigurationDialog.cpp
+++ b/Engine/Plugins/Interchange/Editor/Source/Pipelines/Private/SInterchangePipelineConfigurationDialog.cpp
@@ -361,7 +361,7 @@ FReply SInterchangePipelineConfigurationDialog::OnEditTranslatorSettings()
 			Translator->Translate(*BaseNodeContainer.Get());
 
 			//Refresh the dialog
-			RefreshStack(false);
+			UpdateStack(CurrentStackName);
 
 			Progress.EnterProgressFrame(1.f);
 		}
@@ -402,11 +402,10 @@ TSharedRef<SBox> SInterchangePipelineConfigurationDialog::SpawnPipelineConfigura
 
 	//In case we do not have a valid stack name use the first stack
 	FName FirstStackName = PipelineStacks.Num() > 0 ? PipelineStacks[0].StackName : CurrentStackName;
-	if (!UE::Private::ContainStack(PipelineStacks, CurrentStackName))
+	if (bTestConfigDialog || !UE::Private::ContainStack(PipelineStacks, CurrentStackName))
 	{
 		CurrentStackName = FirstStackName;
 	}
-
 	for(FInterchangeStackInfo& Stack : PipelineStacks)
 	{
 		TSharedPtr<FString> StackNamePtr = MakeShared<FString>(Stack.StackName.ToString());
@@ -418,7 +417,7 @@ TSharedRef<SBox> SInterchangePipelineConfigurationDialog::SpawnPipelineConfigura
 				if (UInterchangePipelineBase* GeneratedPipeline = UE::Interchange::GeneratePipelineInstance(DefaultPipeline))
 				{
 					GeneratedPipeline->TransferAdjustSettings(DefaultPipeline);
-					if (GeneratedPipeline->IsFromReimportOrOverride())
+					if (bTestConfigDialog || GeneratedPipeline->IsFromReimportOrOverride())
 					{
 						//We save the pipeline settings to allow Reset to Default to work
 						GeneratedPipeline->SaveSettings(Stack.StackName);
@@ -751,8 +750,7 @@ TSharedRef<SBox> SInterchangePipelineConfigurationDialog::SpawnPipelineConfigura
 		if (CurrentSelectedPipeline && CurrentSelectedPipeline->IsPropertyChangeNeedRefresh(PropertyChangedEvent))
 		{
 			//Refresh the pipeline
-			constexpr bool bStackSelectionChange = false;
-			RefreshStack(bStackSelectionChange);
+			UpdateStack(CurrentStackName);
 		}
 	});
 	return PipelineConfigurationPanelBox;
@@ -815,8 +813,7 @@ TSharedRef<SBox> SInterchangePipelineConfigurationDialog::SpawnCardsConfiguratio
 						DialogWindow->SetSizeLimits(SizeLimits);
 						DialogWindow->Resize(ClientSize);
 
-						constexpr bool bStackSelectionChange = false;
-						RefreshStack(bStackSelectionChange);
+						UpdateStack(CurrentStackName);
 						GConfig->SetBool(TEXT("InterchangeImportDialogOptions"), TEXT("ShowSettings"), bShowSettings, GEditorPerProjectIni);
 					})
 				[
@@ -857,12 +854,20 @@ TSharedRef<SBox> SInterchangePipelineConfigurationDialog::SpawnCardsConfiguratio
 void SInterchangePipelineConfigurationDialog::Construct(const FArguments& InArgs)
 {
 	LLM_SCOPE_BYNAME(TEXT("Interchange"));
+
+	if (GEditor != nullptr && GEditor->Trans != nullptr)
+	{
+		GEditor->Trans->SetUndoBarrier();
+	}
+
+
 	//Make sure there is a valid default value
 
 	OwnerWindow = InArgs._OwnerWindow;
 	SourceData = InArgs._SourceData;
 	bSceneImport = InArgs._bSceneImport;
 	bReimport = InArgs._bReimport;
+	bTestConfigDialog = InArgs._bTestConfigDialog;
 	PipelineStacks = InArgs._PipelineStacks;
 	OutPipelines = InArgs._OutPipelines;
 	BaseNodeContainer = InArgs._BaseNodeContainer;
@@ -1141,10 +1146,13 @@ void SInterchangePipelineConfigurationDialog::Construct(const FArguments& InArgs
 					[
 						SNew(SPrimaryButton)
 						.Icon(this, &SInterchangePipelineConfigurationDialog::GetImportButtonIcon)
-						.Text(LOCTEXT("InspectorGraphWindow_Import", "Import"))
+						.Text_Lambda([this]() 
+							{
+								return bTestConfigDialog ? LOCTEXT("InspectorGraphWindow_SaveConfig", "Save Config") : LOCTEXT("InspectorGraphWindow_Import", "Import");
+							})
 						.ToolTipText(this, &SInterchangePipelineConfigurationDialog::GetImportButtonTooltip)
 						.IsEnabled(this, &SInterchangePipelineConfigurationDialog::IsImportButtonEnabled)
-						.OnClicked(this, &SInterchangePipelineConfigurationDialog::OnCloseDialog, ECloseEventType::Import)
+						.OnClicked(this, &SInterchangePipelineConfigurationDialog::OnCloseDialog, ECloseEventType::PrimaryButton)
 					]
 					+SUniformGridPanel::Slot(1,0)
 					[
@@ -1155,6 +1163,10 @@ void SInterchangePipelineConfigurationDialog::Construct(const FArguments& InArgs
 							})
 						.HAlign(HAlign_Center)
 						.Text(LOCTEXT("InspectorGraphWindow_Preview", "Preview..."))
+						.IsEnabled_Lambda([this]() 
+							{
+								return !bTestConfigDialog;  
+							})
 						.OnClicked(this, &SInterchangePipelineConfigurationDialog::OnPreviewImport)
 					]
 					+SUniformGridPanel::Slot(2,0)
@@ -1370,7 +1382,7 @@ void SInterchangePipelineConfigurationDialog::ClosePipelineConfiguration(const E
 		bCanceled = true;
 		bImportAll = false;
 	}
-	else //ECloseEventType::Import
+	else //ECloseEventType::PrimaryButton
 	{
 		bCanceled = false;
 		bImportAll = UseSameSettingsForAllCheckBox->IsChecked();
@@ -1405,13 +1417,12 @@ void SInterchangePipelineConfigurationDialog::ClosePipelineConfiguration(const E
 	}
 	OwnerWindow = nullptr;
 
-	// Discard undo history as these transactions should only be valid for the lifetime of the import dialog.
-	if ((GEditor != nullptr) && (GEditor->Trans != nullptr))
+	FInterchangePipelineBaseDetailsCustomization::OnGatherConflictAndExtraInfo.Unbind();
+
+	if (GEditor != nullptr && GEditor->Trans != nullptr)
 	{
-		GEditor->Trans->Reset(LOCTEXT("DiscardHistoryReason", "Discard undo history after closing Interchange Import Dialog."));
+		GEditor->Trans->RemoveUndoBarrier();
 	}
-
-	FInterchangePipelineBaseDetailsCustomization::OnGatherConflictAndExtraInfo.Unbind();
 }
 
 FReply SInterchangePipelineConfigurationDialog::OnKeyDown(const FGeometry& MyGeometry, const FKeyEvent& InKeyEvent)
@@ -1423,10 +1434,12 @@ FReply SInterchangePipelineConfigurationDialog::OnKeyDown(const FGeometry& MyGeo
 	return FReply::Unhandled();
 }
 
-void SInterchangePipelineConfigurationDialog::RefreshStack(bool bStackSelectionChange)
+void SInterchangePipelineConfigurationDialog::UpdateStack(const FName& NewStackName)
 {
 	LLM_SCOPE_BYNAME(TEXT("Interchange"));
 
+	const bool bStackSelectionChange = CurrentStackName != NewStackName;
+
 	//Save current stack settings, we want the same settings when we will go back to the same stack
 	//When doing a reimport we do not want to save the setting because the context have special default
 	//value for some options like: (Import Materials, Import Textures...).
@@ -1435,6 +1448,7 @@ void SInterchangePipelineConfigurationDialog::RefreshStack(bool bStackSelectionC
 	{
 		SaveAllPipelineSettings();
 	}
+	CurrentStackName = NewStackName;
 
 	int32 CurrentPipelineIndex = 0;
 	if (!bStackSelectionChange)
@@ -1528,15 +1542,7 @@ void SInterchangePipelineConfigurationDialog::OnStackSelectionChanged(TSharedPtr
 		return;
 	}
 
-	if (!bReimport)
-	{
-		//Use the stack select by interchange manager
-		SaveAllPipelineSettings();
-	}
-	CurrentStackName = NewStackName;
-
-	constexpr bool bStackSelectionChange = true;
-	RefreshStack(bStackSelectionChange);
+	UpdateStack(NewStackName);
 }
 
 TSharedRef<ITableRow> SInterchangePipelineConfigurationDialog::MakePipelineListRowWidget(
@@ -1735,9 +1741,7 @@ void SInterchangePipelineConfigurationDialog::OnFilterOptionsChanged(ECheckBoxSt
 	}
 	bFilterOptions = bNewCheckValue;
 	//Refresh the pipeline
-	constexpr bool bStackSelectionChange = false;
-	RefreshStack(bStackSelectionChange);
-
+	UpdateStack(CurrentStackName);
 	GConfig->SetBool(TEXT("InterchangeImportDialogOptions"), TEXT("FilterOptions"), bFilterOptions, GEditorPerProjectIni);
 }
 
@@ -1751,8 +1755,7 @@ void SInterchangePipelineConfigurationDialog::OnShowEssentialsChanged(ECheckBoxS
 	}
 	bShowEssentials = bNewCheckValue;
 	//Refresh the pipeline
-	constexpr bool bStackSelectionChange = false;
-	RefreshStack(bStackSelectionChange);
+	UpdateStack(CurrentStackName);
 
 	GConfig->SetBool(TEXT("InterchangeImportDialogOptions"), TEXT("ShowEssentials"), bShowEssentials, GEditorPerProjectIni);
 }
diff --git a/Engine/Plugins/Interchange/Editor/Source/Pipelines/Private/SInterchangePipelineConfigurationDialog.h b/Engine/Plugins/Interchange/Editor/Source/Pipelines/Private/SInterchangePipelineConfigurationDialog.h
--- a/Engine/Plugins/Interchange/Editor/Source/Pipelines/Private/SInterchangePipelineConfigurationDialog.h
+++ b/Engine/Plugins/Interchange/Editor/Source/Pipelines/Private/SInterchangePipelineConfigurationDialog.h
@@ -52,7 +52,7 @@ typedef SListView< TSharedPtr<FInterchangePipelineItemType> > SPipelineListViewT
 enum class ECloseEventType : uint8
 {
 	Cancel,
-	Import,
+	PrimaryButton,
 	WindowClosing
 };
 
@@ -67,6 +67,7 @@ public:
 		SLATE_ARGUMENT(TWeakObjectPtr<UInterchangeSourceData>, SourceData)
 		SLATE_ARGUMENT(bool, bSceneImport)
 		SLATE_ARGUMENT(bool, bReimport)
+		SLATE_ARGUMENT(bool, bTestConfigDialog)
 		SLATE_ARGUMENT(TArray<FInterchangeStackInfo>, PipelineStacks)
 		SLATE_ARGUMENT(TArray<UInterchangePipelineBase*>*, OutPipelines)
 		SLATE_ARGUMENT(TWeakObjectPtr<UInterchangeBaseNodeContainer>, BaseNodeContainer)
@@ -137,7 +138,7 @@ private:
 	// The available stacks
 	TArray<TSharedPtr<FString>> AvailableStacks;
 	void OnStackSelectionChanged(TSharedPtr<FString> String, ESelectInfo::Type);
-	void RefreshStack(bool bStackSelectionChange);
+	void UpdateStack(const FName& NewStackName);
 
 	//////////////////////////////////////////////////////////////////////////
 	// the pipelines list view
@@ -200,6 +201,9 @@ private:
 	bool bCanceled = false;
 	bool bImportAll = false;
 
+	// This is used to denote that the import dialog is opened from the Interchange Test Plan asset to configure the chosen Pipelines.
+	bool bTestConfigDialog = false;
+
 	bool bFilterOptions = false;
 	bool bShowEssentials = false;
 	bool bShowSettings = false;
diff --git a/Engine/Plugins/Interchange/Editor/Source/Pipelines/Public/InterchangePipelineConfigurationGeneric.h b/Engine/Plugins/Interchange/Editor/Source/Pipelines/Public/InterchangePipelineConfigurationGeneric.h
--- a/Engine/Plugins/Interchange/Editor/Source/Pipelines/Public/InterchangePipelineConfigurationGeneric.h
+++ b/Engine/Plugins/Interchange/Editor/Source/Pipelines/Public/InterchangePipelineConfigurationGeneric.h
@@ -35,4 +35,12 @@ protected:
 		, TWeakObjectPtr<UInterchangeBaseNodeContainer> BaseNodeContainer
 		, TWeakObjectPtr <UObject> ReimportAsset
 		, bool bSceneImport) override;
+	virtual EInterchangePipelineConfigurationDialogResult ShowTestPlanConfigurationDialog(TArray<FInterchangeStackInfo>& PipelineStacks
+		, TArray<UInterchangePipelineBase*>& OutPipelines
+		, TWeakObjectPtr<UInterchangeSourceData> SourceData
+		, TWeakObjectPtr <UInterchangeTranslatorBase> Translator
+		, TWeakObjectPtr<UInterchangeBaseNodeContainer> BaseNodeContainer
+		, TWeakObjectPtr <UObject> ReimportAsset
+		, bool bSceneImport
+		, bool bReimport) override;
 };
diff --git a/Engine/Plugins/Tests/InterchangeTests/InterchangeTests.uplugin b/Engine/Plugins/Tests/InterchangeTests/InterchangeTests.uplugin
--- a/Engine/Plugins/Tests/InterchangeTests/InterchangeTests.uplugin
+++ b/Engine/Plugins/Tests/InterchangeTests/InterchangeTests.uplugin
@@ -39,6 +39,10 @@
 		{
 			"Name": "VariantManagerContent",
 			"Enabled": true
+		},
+		{
+			"Name": "EditorScriptingUtilities",
+			"Enabled": true
 		}
 	]
 }
diff --git a/Engine/Plugins/Tests/InterchangeTests/Source/InterchangeTestEditor/InterchangeTestEditor.Build.cs b/Engine/Plugins/Tests/InterchangeTests/Source/InterchangeTestEditor/InterchangeTestEditor.Build.cs
--- a/Engine/Plugins/Tests/InterchangeTests/Source/InterchangeTestEditor/InterchangeTestEditor.Build.cs
+++ b/Engine/Plugins/Tests/InterchangeTests/Source/InterchangeTestEditor/InterchangeTestEditor.Build.cs
@@ -17,6 +17,7 @@ namespace UnrealBuildTool.Rules
 					"Engine",
 					"EditorStyle",
 					"InputCore",
+					"InterchangeCore",
 					"InterchangeTests",
 					"LevelEditor",
 					"PropertyEditor",
@@ -31,6 +32,7 @@ namespace UnrealBuildTool.Rules
 					"AssetDefinition",
 					"AssetTools",
 					"WorkspaceMenuStructure",
+					"EditorScriptingUtilities",
 					"SlateCore",
 					"Slate"
 				}
diff --git a/Engine/Plugins/Tests/InterchangeTests/Source/InterchangeTestEditor/Private/InterchangeTestEditorModule.cpp b/Engine/Plugins/Tests/InterchangeTests/Source/InterchangeTestEditor/Private/InterchangeTestEditorModule.cpp
--- a/Engine/Plugins/Tests/InterchangeTests/Source/InterchangeTestEditor/Private/InterchangeTestEditorModule.cpp
+++ b/Engine/Plugins/Tests/InterchangeTests/Source/InterchangeTestEditor/Private/InterchangeTestEditorModule.cpp
@@ -2,7 +2,9 @@
 
 #include "InterchangeTestEditorModule.h"
 #include "InterchangeTestFunctionLayout.h"
+#include "InterchangeTestPlanPipelineSettingsLayout.h"
 #include "InterchangeTestFunction.h"
+#include "InterchangeTestPlanPipelineSettings.h"
 #include "Modules/ModuleManager.h"
 #include "PropertyEditorModule.h"
 
@@ -28,6 +30,13 @@ void FInterchangeTestEditorModule::StartupModule()
 	FPropertyEditorModule& PropertyModule = FModuleManager::LoadModuleChecked<FPropertyEditorModule>("PropertyEditor");
 	PropertyModule.RegisterCustomPropertyTypeLayout(FInterchangeTestFunction::StaticStruct()->GetFName(),
 	FOnGetPropertyTypeCustomizationInstance::CreateStatic(&FInterchangeTestFunctionLayout::MakeInstance));
+
+	PropertyModule.RegisterCustomPropertyTypeLayout(FInterchangeTestPlanPipelineSettings::StaticStruct()->GetFName(),
+	FOnGetPropertyTypeCustomizationInstance::CreateStatic(&FInterchangeTestPlanPipelineSettingsLayout::MakeInstance));
+
+	PropertyModule.RegisterCustomPropertyTypeLayout(FInterchangeImportIntoLevelBoolean::StaticStruct()->GetFName(),
+	FOnGetPropertyTypeCustomizationInstance::CreateStatic(&FInterchangeTestPlanImportIntoLevelLayout::MakeInstance));
+
 	PropertyModule.NotifyCustomizationModuleChanged();
 }
 
@@ -38,6 +47,8 @@ void FInterchangeTestEditorModule::ShutdownModule()
 	{
 		FPropertyEditorModule& PropertyModule = FModuleManager::LoadModuleChecked<FPropertyEditorModule>("PropertyEditor");
 		PropertyModule.UnregisterCustomPropertyTypeLayout(FInterchangeTestFunction::StaticStruct()->GetFName());
+		PropertyModule.UnregisterCustomPropertyTypeLayout(FInterchangeTestPlanPipelineSettings::StaticStruct()->GetFName());
+		PropertyModule.UnregisterCustomPropertyTypeLayout(FInterchangeImportIntoLevelBoolean::StaticStruct()->GetFName());
 	}
 }
 
diff --git a/Engine/Plugins/Tests/InterchangeTests/Source/InterchangeTestEditor/Private/InterchangeTestPlanPipelineSettingsLayout.cpp b/Engine/Plugins/Tests/InterchangeTests/Source/InterchangeTestEditor/Private/InterchangeTestPlanPipelineSettingsLayout.cpp
new file mode 100644
--- /dev/null
+++ b/Engine/Plugins/Tests/InterchangeTests/Source/InterchangeTestEditor/Private/InterchangeTestPlanPipelineSettingsLayout.cpp
@@ -0,0 +1,269 @@
+// Copyright Epic Games, Inc. All Rights Reserved.
+#include "InterchangeTestPlanPipelineSettingsLayout.h"
+#include "InterchangeTestPlanPipelineSettings.h"
+#include "InterchangeImportTestStepBase.h"
+#include "InterchangeImportTestStepImport.h"
+#include "InterchangePipelineBase.h"
+
+#include "DetailWidgetRow.h"
+#include "DetailLayoutBuilder.h"
+#include "Dialogs/Dialogs.h"
+#include "Editor.h"
+#include "HAL/IConsoleManager.h"
+#include "IDetailChildrenBuilder.h"
+#include "IPropertyUtilities.h"
+#include "Misc/ConfigCacheIni.h"
+#include "Widgets/Input/SButton.h"
+#include "Widgets/Input/SCheckBox.h"
+
+#define LOCTEXT_NAMESPACE "InterchangeTestPipelineSettings"
+
+static bool GInterchangeTestPlanCanEditCustomPipelines = false;
+static FAutoConsoleVariableRef CCvarInterchangeDefaultShowEssentialsView(
+	TEXT("Interchange.TestPlan.CanEditCustomPipelines"),
+	GInterchangeTestPlanCanEditCustomPipelines,
+	TEXT("Can the Pipelines in the Test Plan asset be edited."),
+	ECVF_Default);
+
+
+
+/////////// FInterchangeTestPlanPipelineSettingsLayout ///////////
+
+TSharedRef<IPropertyTypeCustomization> FInterchangeTestPlanPipelineSettingsLayout::MakeInstance()
+{
+	return MakeShared<FInterchangeTestPlanPipelineSettingsLayout>();
+}
+
+FInterchangeTestPlanPipelineSettingsLayout::FInterchangeTestPlanPipelineSettingsLayout()
+{
+}
+
+FInterchangeTestPlanPipelineSettingsLayout::~FInterchangeTestPlanPipelineSettingsLayout()
+{
+	FEditorDelegates::PostUndoRedo.RemoveAll(this);
+}
+
+void FInterchangeTestPlanPipelineSettingsLayout::CustomizeHeader(TSharedRef<IPropertyHandle> StructPropertyHandle, FDetailWidgetRow& HeaderRow, IPropertyTypeCustomizationUtils& StructCustomizationUtils)
+{
+	FEditorDelegates::PostUndoRedo.AddSP(this, &FInterchangeTestPlanPipelineSettingsLayout::RefreshLayout);
+
+	StructProperty = StructPropertyHandle;
+	PropertyUtilities = StructCustomizationUtils.GetPropertyUtilities();
+
+	HeaderRow
+	.NameContent()
+	[
+		StructPropertyHandle->CreatePropertyNameWidget()
+	]
+	.ValueContent()
+	[
+		SNew(SHorizontalBox)
+			+ SHorizontalBox::Slot()
+			.AutoWidth()
+			[
+				SNew(SButton)
+					.Text_Lambda([this]() 
+						{
+							return LOCTEXT("PipelineSettingsLayout_EditDefaults", "Edit Pipeline Settings");
+						})
+					.IsEnabled_Lambda([this]() 
+						{
+							if (FInterchangeTestPlanPipelineSettings* PipelineData = GetStruct())
+							{
+								return PipelineData->CanEditPipelineSettings();
+							}
+
+							return false;
+						})
+					.OnClicked(this, &FInterchangeTestPlanPipelineSettingsLayout::EditPipelineSettings)
+			]
+			+ SHorizontalBox::Slot()
+			.AutoWidth()
+			[
+				SNew(SButton)
+					.Text_Lambda([this]()
+						{
+							return LOCTEXT("PipelineSettingsLayout_Clear", "Clear");
+						})
+					.OnClicked(this, &FInterchangeTestPlanPipelineSettingsLayout::ClearModifiedPipelineSettings)
+			]
+	];
+}
+
+void FInterchangeTestPlanPipelineSettingsLayout::CustomizeChildren(TSharedRef<IPropertyHandle> StructPropertyHandle, IDetailChildrenBuilder& StructBuilder, IPropertyTypeCustomizationUtils& StructCustomizationUtils)
+{
+	if (!StructPropertyHandle->IsValidHandle())
+	{
+		return;
+	}
+
+	uint32 NumChildren = 0;
+	StructPropertyHandle->GetNumChildren(NumChildren);
+
+	for (uint32 ChildIndex = 0; ChildIndex < NumChildren; ChildIndex++)
+	{
+		TSharedRef<IPropertyHandle> ChildHandle = StructPropertyHandle->GetChildHandle(ChildIndex).ToSharedRef();
+		if (ChildHandle->GetProperty()->GetFName() == GET_MEMBER_NAME_CHECKED(FInterchangeTestPlanPipelineSettings, CustomPipelines))
+		{
+			// Only show the property in the editor when the CVar is enabled.
+			if (GInterchangeTestPlanCanEditCustomPipelines)
+			{
+				StructBuilder.AddProperty(ChildHandle);
+			}
+		}
+		else 
+		{
+			StructBuilder.AddProperty(ChildHandle);
+		}
+	}
+}
+
+void FInterchangeTestPlanPipelineSettingsLayout::RefreshLayout()
+{
+	if (PropertyUtilities.IsValid())
+	{
+		PropertyUtilities->ForceRefresh();
+	}
+}
+
+FReply FInterchangeTestPlanPipelineSettingsLayout::EditPipelineSettings()
+{
+	if (FInterchangeTestPlanPipelineSettings* PipelineSettings = GetStruct())
+	{
+		if (UInterchangeImportTestStepBase* TestStep = PipelineSettings->ParentTestStep.Get())
+		{
+			TestStep->EditPipelineSettings();
+		}
+	}
+	return FReply::Handled();
+}
+
+FReply FInterchangeTestPlanPipelineSettingsLayout::ClearModifiedPipelineSettings()
+{
+	if (FInterchangeTestPlanPipelineSettings* PipelineSettings = GetStruct())
+	{
+		if (UInterchangeImportTestStepBase* TestStep = PipelineSettings->ParentTestStep.Get())
+		{
+			TestStep->ClearPipelineSettings();
+		}
+	}
+	return FReply::Handled();
+}
+
+FInterchangeTestPlanPipelineSettings* FInterchangeTestPlanPipelineSettingsLayout::GetStruct() const
+{
+	// Get address of the FInterchangeTestPlanPipelineSettings struct being viewed.
+	// We only ever expect the property handle to be linked to a single instance.
+
+	void* StructPtr = nullptr;
+	StructProperty->GetValueData(StructPtr);
+	return static_cast<FInterchangeTestPlanPipelineSettings*>(StructPtr);
+}
+
+
+/////////// FInterchangeTestPlanImportIntoLevelLayout ///////////
+TSharedRef<IPropertyTypeCustomization> FInterchangeTestPlanImportIntoLevelLayout::MakeInstance()
+{
+	return MakeShared<FInterchangeTestPlanImportIntoLevelLayout>();
+}
+
+FInterchangeTestPlanImportIntoLevelLayout::FInterchangeTestPlanImportIntoLevelLayout()
+{
+}
+
+FInterchangeTestPlanImportIntoLevelLayout::~FInterchangeTestPlanImportIntoLevelLayout()
+{
+	FEditorDelegates::PostUndoRedo.RemoveAll(this);
+}
+
+void FInterchangeTestPlanImportIntoLevelLayout::CustomizeHeader(TSharedRef<IPropertyHandle> StructPropertyHandle, FDetailWidgetRow& HeaderRow, IPropertyTypeCustomizationUtils& StructCustomizationUtils)
+{
+	FEditorDelegates::PostUndoRedo.AddSP(this, &FInterchangeTestPlanImportIntoLevelLayout::RefreshLayout);
+
+	StructProperty = StructPropertyHandle;
+	PropertyUtilities = StructCustomizationUtils.GetPropertyUtilities();
+
+	HeaderRow
+		.NameContent()
+		[
+			StructPropertyHandle->CreatePropertyNameWidget()
+		]
+		.ValueContent()
+		[
+			SNew(SCheckBox)
+			.IsChecked_Lambda([this]()
+				{
+					if (FInterchangeImportIntoLevelBoolean* MessageBoxBoolean = GetStruct())
+					{
+						return MessageBoxBoolean->GetValue() ? ECheckBoxState::Checked : ECheckBoxState::Unchecked;
+					}
+					return ECheckBoxState::Unchecked;
+				})
+			.OnCheckStateChanged_Lambda([this](ECheckBoxState NewState)
+				{
+					if (FInterchangeImportIntoLevelBoolean* MessageBoxBoolean = GetStruct())
+					{
+						const bool bShouldShowMessageBox = MessageBoxBoolean->ShouldValueChangeShowMessageBox.IsBound() && MessageBoxBoolean->ShouldValueChangeShowMessageBox.Execute();
+						if (bShouldShowMessageBox)
+						{
+							FSuppressableWarningDialog::FSetupInfo DialogSetupInfo = FSuppressableWarningDialog::FSetupInfo(
+								MessageBoxBoolean->MessageBoxText,
+								MessageBoxBoolean->MessageBoxTitle,
+								MessageBoxBoolean->MessageBoxShowDialogINISetting);
+
+							DialogSetupInfo.ConfirmText = LOCTEXT("TestPlanMessageBoxBoolean_Confirm", "Yes");
+							DialogSetupInfo.CancelText = LOCTEXT("TestPlanMessageBoxBoolean_Cancel", "No");
+							DialogSetupInfo.bDefaultToSuppressInTheFuture = false;
+							DialogSetupInfo.CheckBoxText = LOCTEXT("TestPlanMessageBoxBoolean_CheckBoxText", "Don't show this dialog again");
+							FSuppressableWarningDialog SuppressableWarningDialog(DialogSetupInfo);
+							FSuppressableWarningDialog::EResult DialogResult = SuppressableWarningDialog.ShowModal();
+
+							if (DialogResult != FSuppressableWarningDialog::EResult::Cancel)
+							{
+								MessageBoxBoolean->SetValue(NewState == ECheckBoxState::Checked);
+							}
+						}
+						else 
+						{
+							MessageBoxBoolean->SetValue(NewState == ECheckBoxState::Checked);
+						}
+					}
+				})
+		];
+}
+
+void FInterchangeTestPlanImportIntoLevelLayout::CustomizeChildren(TSharedRef<IPropertyHandle> StructPropertyHandle, IDetailChildrenBuilder& StructBuilder, IPropertyTypeCustomizationUtils& StructCustomizationUtils)
+{
+	if (!StructPropertyHandle->IsValidHandle())
+	{
+		return;
+	}
+
+	uint32 NumChildren = 0;
+	StructPropertyHandle->GetNumChildren(NumChildren);
+
+	for (uint32 ChildIndex = 0; ChildIndex < NumChildren; ChildIndex++)
+	{
+		StructBuilder.AddProperty(StructPropertyHandle->GetChildHandle(ChildIndex).ToSharedRef());
+	}
+}
+
+void FInterchangeTestPlanImportIntoLevelLayout::RefreshLayout()
+{
+	if (PropertyUtilities.IsValid())
+	{
+		PropertyUtilities->ForceRefresh();
+	}
+}
+
+FInterchangeImportIntoLevelBoolean* FInterchangeTestPlanImportIntoLevelLayout::GetStruct() const
+{
+	// Get address of the FInterchangeTestPlanImportIntoLevelLayout struct being viewed.
+	// We only ever expect the property handle to be linked to a single instance.
+
+	void* StructPtr = nullptr;
+	StructProperty->GetValueData(StructPtr);
+	return static_cast<FInterchangeImportIntoLevelBoolean*>(StructPtr);
+}
+
+#undef LOCTEXT_NAMESPACE 
\ No newline at end of file
diff --git a/Engine/Plugins/Tests/InterchangeTests/Source/InterchangeTestEditor/Private/InterchangeTestPlanPipelineSettingsLayout.h b/Engine/Plugins/Tests/InterchangeTests/Source/InterchangeTestEditor/Private/InterchangeTestPlanPipelineSettingsLayout.h
new file mode 100644
--- /dev/null
+++ b/Engine/Plugins/Tests/InterchangeTests/Source/InterchangeTestEditor/Private/InterchangeTestPlanPipelineSettingsLayout.h
@@ -0,0 +1,62 @@
+// Copyright Epic Games, Inc. All Rights Reserved.
+#pragma once
+#include "IPropertyTypeCustomization.h"
+
+class FDetailWidgetRow;
+class IDetailChildrenBuilder;
+class IPropertyUtilities;
+struct FInterchangeTestPlanPipelineSettings;
+struct FInterchangeImportIntoLevelBoolean;
+class FReply;
+
+class FInterchangeTestPlanPipelineSettingsLayout : public IPropertyTypeCustomization
+{
+public:
+	/** Makes a new instance of this layout class for a specific detail view requesting it */
+	static TSharedRef<IPropertyTypeCustomization> MakeInstance();
+
+	FInterchangeTestPlanPipelineSettingsLayout();
+	virtual ~FInterchangeTestPlanPipelineSettingsLayout();
+
+	// IPropertyTypeCustomization interface
+	virtual void CustomizeHeader(TSharedRef<IPropertyHandle> StructPropertyHandle, FDetailWidgetRow& HeaderRow, IPropertyTypeCustomizationUtils& StructCustomizationUtils) override;
+	virtual void CustomizeChildren(TSharedRef<IPropertyHandle> StructPropertyHandle, IDetailChildrenBuilder& StructBuilder, IPropertyTypeCustomizationUtils& StructCustomizationUtils) override;
+
+private:
+	void RefreshLayout();
+
+	FReply EditPipelineSettings();
+	FReply ClearModifiedPipelineSettings();
+
+private:
+	FInterchangeTestPlanPipelineSettings* GetStruct() const;
+
+private:
+	TSharedPtr<IPropertyUtilities> PropertyUtilities;
+	TSharedPtr<IPropertyHandle> StructProperty;
+
+};
+
+class FInterchangeTestPlanImportIntoLevelLayout : public IPropertyTypeCustomization
+{
+public:
+	/** Makes a new instance of this layout class for a specific detail view requesting it */
+	static TSharedRef<IPropertyTypeCustomization> MakeInstance();
+
+	FInterchangeTestPlanImportIntoLevelLayout();
+	virtual ~FInterchangeTestPlanImportIntoLevelLayout();
+
+	// IPropertyTypeCustomization interface
+	virtual void CustomizeHeader(TSharedRef<IPropertyHandle> StructPropertyHandle, FDetailWidgetRow& HeaderRow, IPropertyTypeCustomizationUtils& StructCustomizationUtils) override;
+	virtual void CustomizeChildren(TSharedRef<IPropertyHandle> StructPropertyHandle, IDetailChildrenBuilder& StructBuilder, IPropertyTypeCustomizationUtils& StructCustomizationUtils) override;
+
+private:
+	void RefreshLayout();
+private:
+	FInterchangeImportIntoLevelBoolean* GetStruct() const;
+
+private:
+	TSharedPtr<IPropertyUtilities> PropertyUtilities;
+	TSharedPtr<IPropertyHandle> StructProperty;
+
+};
\ No newline at end of file
diff --git a/Engine/Plugins/Tests/InterchangeTests/Source/InterchangeTests/Private/InterchangeImportTestStepImport.cpp b/Engine/Plugins/Tests/InterchangeTests/Source/InterchangeTests/Private/InterchangeImportTestStepImport.cpp
--- a/Engine/Plugins/Tests/InterchangeTests/Source/InterchangeTests/Private/InterchangeImportTestStepImport.cpp
+++ b/Engine/Plugins/Tests/InterchangeTests/Source/InterchangeTests/Private/InterchangeImportTestStepImport.cpp
@@ -2,16 +2,36 @@
 
 #include "InterchangeImportTestStepImport.h"
 
+#include "Dialogs/Dialogs.h"
+#include "Editor.h"
 #include "Engine/Engine.h"
 #include "Engine/Level.h"
 #include "Engine/World.h"
 #include "GameFramework/WorldSettings.h"
 #include "HAL/FileManager.h"
 #include "InterchangeImportTestData.h"
+#include "InterchangeManager.h"
+#include "InterchangeProjectSettings.h"
+#include "InterchangeTestsLog.h"
 #include "Misc/PackageName.h"
+#include "Misc/ScopedSlowTask.h"
+#include "ScopedTransaction.h"
 
 #include UE_INLINE_GENERATED_CPP_BY_NAME(InterchangeImportTestStepImport)
 
+#define LOCTEXT_NAMESPACE "InterchangeImportTestStepImport"
+
+UInterchangeImportTestStepImport::UInterchangeImportTestStepImport(const FObjectInitializer& InObjectInitializer)
+{
+	PipelineSettings.ParentTestStep = this;
+
+	bMessageBoxImportIntoLevel = FInterchangeImportIntoLevelBoolean(bImportIntoLevel);
+	bMessageBoxImportIntoLevel.MessageBoxTitle = LOCTEXT("TestStepImport_ImportIntoLevelTitle", "Import Type Changed");
+	bMessageBoxImportIntoLevel.MessageBoxText = LOCTEXT("TestStepImport_ImportIntoLevelText", "Import type is changed. This will delete all the modifications made.\nNOTE: This change is irreversible.\nDo you still wish to continue? ");
+	bMessageBoxImportIntoLevel.MessageBoxShowDialogINISetting = "InterchangeImportTestPlanImportIntoLevelWarning";
+	bMessageBoxImportIntoLevel.OnValueChanged.BindUObject(this, &UInterchangeImportTestStepImport::OnImportIntoLevelChanged);
+	bMessageBoxImportIntoLevel.ShouldValueChangeShowMessageBox.BindUObject(this, &UInterchangeImportTestStepImport::ShouldImportIntoLevelChangeRequireMessageBox);
+}
 
 TTuple<UE::Interchange::FAssetImportResultPtr, UE::Interchange::FSceneImportResultPtr> UInterchangeImportTestStepImport::StartStep(FInterchangeImportTestData& Data)
 {
@@ -59,10 +79,21 @@ TTuple<UE::Interchange::FAssetImportResultPtr, UE::Interchange::FSceneImportResu
 	UE::Interchange::FScopedSourceData ScopedSourceData(SourceFile.FilePath);
 
 	FImportAssetParameters Params;
-	Params.OverridePipelines.Reserve(PipelineStack.Num());
-	for(TObjectPtr<UInterchangePipelineBase> Pipeline : PipelineStack)
+	if (bUseOverridePipelineStack)
+	{
+		Params.OverridePipelines.Reserve(PipelineStack.Num());
+		for (TObjectPtr<UInterchangePipelineBase> Pipeline : PipelineStack)
+		{
+			Params.OverridePipelines.Add(Pipeline);
+		}
+	}
+	else if (PipelineSettings.CustomPipelines.Num() > 0)
 	{
-		Params.OverridePipelines.Add(Pipeline);
+		Params.OverridePipelines.Reserve(PipelineSettings.CustomPipelines.Num());
+		for (TObjectPtr<UInterchangePipelineBase> Pipeline : PipelineSettings.CustomPipelines)
+		{
+			Params.OverridePipelines.Add(Pipeline);
+		}
 	}
 	Params.bIsAutomated = true;
 
@@ -113,33 +144,282 @@ FTestStepResults UInterchangeImportTestStepImport::FinishStep(FInterchangeImport
 	// Run all the tests
 	bool bSuccess = PerformTests(Data, ExecutionInfo);
 
-	if (TransientWorld)
+if (TransientWorld)
+{
+	// Only keep assets as result objects since the world and its actors are being destroyed
+	TArray<UObject*> ResultObjects = MoveTemp(Data.ResultObjects);
+	Data.ResultObjects.Reset(ResultObjects.Num());
+
+	for (UObject* Object : ResultObjects)
+	{
+		if (!Object->IsA<AActor>())
+		{
+			Data.ResultObjects.Add(Object);
+		}
+	}
+
+	// Now delete world
+	GEngine->DestroyWorldContext(TransientWorld.Get());
+	TransientWorld->DestroyWorld(true);
+	TransientWorld.Reset();
+}
+
+Results.bTestStepSuccess = bSuccess;
+return Results;
+}
+
+
+FString UInterchangeImportTestStepImport::GetContextString() const
+{
+	return FString(TEXT("Importing ")) + FPaths::GetCleanFilename(SourceFile.FilePath);
+}
+
+bool UInterchangeImportTestStepImport::CanEditPipelineSettings() const
+{
+	if (SourceFile.FilePath.IsEmpty())
+	{
+		return false;
+	}
+
+	if (bUseOverridePipelineStack)
+	{
+		if (PipelineStack.IsEmpty())
+		{
+			return false;
+		}
+
+		for (const UInterchangePipelineBase* Pipeline : PipelineStack)
+		{
+			if (!Pipeline)
+			{
+				return false;
+			}
+		}
+	}
+
+	return true;
+}
+
+void UInterchangeImportTestStepImport::EditPipelineSettings()
+{
+#if WITH_EDITOR
+	if (SourceFile.FilePath.IsEmpty())
+	{
+		return;
+	}
+
+	UE::Interchange::FScopedSourceData ScopedSourceData(SourceFile.FilePath);
+	UE::Interchange::FScopedTranslator ScopedTranslator(ScopedSourceData.GetSourceData());
+	if (!ScopedTranslator.GetTranslator())
+	{
+		UE_LOG(LogInterchangeTests, Error, TEXT("Cannot import file. The source data is not supported. Try enabling the [%s] extension for Interchange."), *FPaths::GetExtension(ScopedSourceData.GetSourceData()->GetFilename()));
+		return;
+	}
+
+	const FInterchangeImportSettings& InterchangeImportSettings = FInterchangeProjectSettingsUtils::GetDefaultImportSettings(bImportIntoLevel);
+	if (InterchangeImportSettings.PipelineStacks.Num() == 0)
+	{
+		UE_LOG(LogInterchangeTests, Error, TEXT("Failed to configure pipelines. There is no pipeline stack defined for the content import type."));
+		return;
+	}
+
+	if (!InterchangeImportSettings.PipelineStacks.Contains(InterchangeImportSettings.DefaultPipelineStack))
 	{
-		// Only keep assets as result objects since the world and its actors are being destroyed
-		TArray<UObject*> ResultObjects = MoveTemp(Data.ResultObjects);
-		Data.ResultObjects.Reset(ResultObjects.Num());
+		FInterchangeImportSettings& MutableInterchangeImportSettings = FInterchangeProjectSettingsUtils::GetMutableDefaultImportSettings(bImportIntoLevel);
+		TArray<FName> Keys;
+		MutableInterchangeImportSettings.PipelineStacks.GetKeys(Keys);
+		check(Keys.Num() > 0);
+		MutableInterchangeImportSettings.DefaultPipelineStack = Keys[0];
+	}
+
+	UE::Interchange::FScopedBaseNodeContainer ScopedBaseNodeContainer;
 
-		for (UObject* Object : ResultObjects)
+	UInterchangePipelineConfigurationBase* RegisteredPipelineConfiguration = nullptr;
+	TSoftClassPtr <UInterchangePipelineConfigurationBase> ImportDialogClass = InterchangeImportSettings.ImportDialogClass;
+	if (ImportDialogClass.IsValid())
+	{
+		UClass* PipelineConfigurationClass = ImportDialogClass.LoadSynchronous();
+		if (PipelineConfigurationClass)
 		{
-			if (!Object->IsA<AActor>())
+			RegisteredPipelineConfiguration = NewObject<UInterchangePipelineConfigurationBase>(GetTransientPackage(), PipelineConfigurationClass, NAME_None, RF_NoFlags);
+			if (RegisteredPipelineConfiguration == nullptr)
 			{
-				Data.ResultObjects.Add(Object);
+				UE_LOG(LogInterchangeTests, Error, TEXT("Failed to create a pipeline configuration object."));
+				return;
 			}
 		}
+	}
+
+	auto AdjustPipelineSettingForContext = [this, &ScopedBaseNodeContainer](UInterchangePipelineBase* Pipeline)
+		{
+			EInterchangePipelineContext Context;
+			Context = bImportIntoLevel ? EInterchangePipelineContext::SceneImport : EInterchangePipelineContext::AssetImport;
+			FInterchangePipelineContextParams ContextParams;
+			ContextParams.ContextType = Context;
+			ContextParams.BaseNodeContainer = ScopedBaseNodeContainer.GetBaseNodeContainer();
+			Pipeline->SetFromReimportOrOverride(bUseOverridePipelineStack);
+			Pipeline->AdjustSettingsForContext(ContextParams);
+		};
 
-		// Now delete world
-		GEngine->DestroyWorldContext(TransientWorld.Get());
-		TransientWorld->DestroyWorld(true);
-		TransientWorld.Reset();
+	const bool bCanTranslate = UInterchangeManager::GetInterchangeManager().CanUseTranslator(ScopedTranslator.GetTranslator());
+	if (bCanTranslate)
+	{
+		FScopedSlowTask Progress(2.f, LOCTEXT("TranslatingSourceFile...", "Translating source file..."));
+		Progress.MakeDialog();
+		Progress.EnterProgressFrame(1.f);
+
+		UInterchangeBaseNodeContainer& BaseNodeContainer = *(ScopedBaseNodeContainer.GetBaseNodeContainer());
+		ScopedTranslator.GetTranslator()->SetResultsContainer(NewObject<UInterchangeResultsContainer>(GetTransientPackage(), NAME_None));
+		ScopedTranslator.GetTranslator()->Translate(BaseNodeContainer);
+
+		Progress.EnterProgressFrame(1.f);
 	}
 
-	Results.bTestStepSuccess = bSuccess;
-	return Results;
+	const TMap<FName, FInterchangePipelineStack>& DefaultPipelineStacks = InterchangeImportSettings.PipelineStacks;
+	TArray<FInterchangeStackInfo> InPipelineStacks;
+	TArray<UInterchangePipelineBase*> OutPipelines;
+
+	if (bUseOverridePipelineStack && !PipelineStack.IsEmpty())
+	{
+		FInterchangeStackInfo& StackInfo = InPipelineStacks.AddDefaulted_GetRef();
+		StackInfo.StackName = FName(TEXT("OverridePipelineStack"));
+		for (int32 PipelineClassIndex = 0; PipelineClassIndex < PipelineStack.Num(); ++PipelineClassIndex)
+		{
+			if (UInterchangePipelineBase* GeneratedPipeline = UE::Interchange::GeneratePipelineInstance(PipelineStack[PipelineClassIndex]))
+			{
+				AdjustPipelineSettingForContext(GeneratedPipeline);
+				StackInfo.Pipelines.Add(GeneratedPipeline);
+			}
+		}
+
+		if (StackInfo.Pipelines.IsEmpty())
+		{
+			UE_LOG(LogInterchangeTests, Error, TEXT("Failed to configure pipelines. There are no pipelines in the override stack"));
+			return;
+		}
+	}
+	else if (PipelineSettings.CustomPipelines.Num() > 0)
+	{
+		FInterchangeStackInfo& StackInfo = InPipelineStacks.AddDefaulted_GetRef();
+		StackInfo.StackName = FName(TEXT("CustomPipelineStack"));
+		for (int32 PipelineIndex = 0; PipelineIndex < PipelineSettings.CustomPipelines.Num(); ++PipelineIndex)
+		{
+			if (UInterchangePipelineBase* GeneratedPipeline = UE::Interchange::GeneratePipelineInstance(PipelineSettings.CustomPipelines[PipelineIndex]))
+			{
+				AdjustPipelineSettingForContext(GeneratedPipeline);
+				StackInfo.Pipelines.Add(GeneratedPipeline);
+			}
+		}
+	}
+
+	for (const TPair<FName, FInterchangePipelineStack>& PipelineStackInfo : DefaultPipelineStacks)
+	{
+		FName StackName = PipelineStackInfo.Key;
+		FInterchangeStackInfo& StackInfo = InPipelineStacks.AddDefaulted_GetRef();
+		StackInfo.StackName = StackName;
+
+		const FInterchangePipelineStack& DefPipelineStack = PipelineStackInfo.Value;
+		const TArray<FSoftObjectPath>* Pipelines = &DefPipelineStack.Pipelines;
+
+		// If applicable, check to see if a specific pipeline stack is associated with this translator
+		for (const FInterchangeTranslatorPipelines& TranslatorPipelines : DefPipelineStack.PerTranslatorPipelines)
+		{
+			const UClass* TranslatorClass = TranslatorPipelines.Translator.LoadSynchronous();
+			if (ScopedTranslator.GetTranslator() && ScopedTranslator.GetTranslator()->IsA(TranslatorClass))
+			{
+				Pipelines = &TranslatorPipelines.Pipelines;
+				break;
+			}
+		}
+
+		for (int32 PipelineIndex = 0; PipelineIndex < Pipelines->Num(); ++PipelineIndex)
+		{
+			if (UInterchangePipelineBase* GeneratedPipeline = UE::Interchange::GeneratePipelineInstance((*Pipelines)[PipelineIndex]))
+			{
+				AdjustPipelineSettingForContext(GeneratedPipeline);
+				StackInfo.Pipelines.Add(GeneratedPipeline);
+			}
+		}
+	}
+
+	//Show the dialog, a plugin should have register this dialog. We use a plugin to be able to use editor code when doing UI
+	EInterchangePipelineConfigurationDialogResult DialogResult = RegisteredPipelineConfiguration->ScriptedShowTestPlanConfigurationDialog(InPipelineStacks, OutPipelines, ScopedSourceData.GetSourceData(), ScopedTranslator.GetTranslator(), ScopedBaseNodeContainer.GetBaseNodeContainer(), nullptr, bImportIntoLevel);
+
+	if (DialogResult == EInterchangePipelineConfigurationDialogResult::SaveConfig)
+	{
+		if (bUseOverridePipelineStack)
+		{
+			FScopedTransaction ScopedTransaction(LOCTEXT("UpdatePipelineSettings", "Update Pipeline Settings"));
+			Modify();
+
+			PipelineStack.Empty(OutPipelines.Num());
+			for (UInterchangePipelineBase* Pipeline : OutPipelines)
+			{
+				if (UInterchangePipelineBase* InstancedPipeline = DuplicateObject<UInterchangePipelineBase>(Pipeline, GetOuter()))
+				{
+					InstancedPipeline->SetFlags(EObjectFlags::RF_Transactional);
+					PipelineStack.Add(InstancedPipeline);
+				}
+			}
+		}
+		else
+		{
+			PipelineSettings.UpdatePipelines(OutPipelines);
+		}
+	}
+#endif 
 }
 
+void UInterchangeImportTestStepImport::ClearPipelineSettings()
+{
+#if WITH_EDITOR
+	
+	if (bUseOverridePipelineStack)
+	{
+		FScopedTransaction ScopedTransaction(LOCTEXT("ClearPipelineSettings", "Clear Pipeline Settings"));
+		Modify();
+		PipelineStack.Empty();
+	}
+	else
+	{
+		PipelineSettings.ClearPipelines();
+	}
+#endif
+}
 
-FString UInterchangeImportTestStepImport::GetContextString() const
+void UInterchangeImportTestStepImport::OnImportIntoLevelChanged(bool bNewImportIntoLevel)
 {
-	return FString(TEXT("Importing ")) + FPaths::GetCleanFilename(SourceFile.FilePath);
+	bImportIntoLevel = bNewImportIntoLevel;
+	PipelineSettings.ClearPipelines();
+}
+
+bool UInterchangeImportTestStepImport::ShouldImportIntoLevelChangeRequireMessageBox() const
+{
+	return !bUseOverridePipelineStack && PipelineSettings.IsUsingModifiedSettings();
+}
+
+#if WITH_EDITOR
+void UInterchangeImportTestStepImport::PostEditChangeProperty(FPropertyChangedEvent& PropertyChangedEvent)
+{
+	if (PropertyChangedEvent.Property->GetFName() == GET_MEMBER_NAME_CHECKED(UInterchangeImportTestStepImport, bUseOverridePipelineStack))
+	{
+		if (bUseOverridePipelineStack && !PipelineSettings.CustomPipelines.IsEmpty())
+		{
+			FSuppressableWarningDialog::FSetupInfo DialogSetupInfo = FSuppressableWarningDialog::FSetupInfo(
+				LOCTEXT("TestStepImport_UseOverridePipelinesText", "You are using override pipelines now. This will erase all the previous modifications made to the default pipeline settings."),
+				LOCTEXT("TestStepImport_UseOverridePipelinesTitle", "Using Override Pipelines"),
+				"InterchangeImportTestPlanUseOverridePipelinesWarning");
+
+			DialogSetupInfo.ConfirmText = LOCTEXT("TestStepImport_UseOverridePipelines_Confirm", "OK");
+			DialogSetupInfo.bDefaultToSuppressInTheFuture = false;
+			DialogSetupInfo.CheckBoxText = LOCTEXT("TestStepImport_UseOverridePipelines_CheckBoxText", "Don't show this dialog again");
+			FSuppressableWarningDialog SuppressableWarningDialog(DialogSetupInfo);
+			SuppressableWarningDialog.ShowModal();
+
+			PipelineSettings.ClearPipelines();
+		}
+	}
 }
+#endif
 
+#undef LOCTEXT_NAMESPACE
\ No newline at end of file
diff --git a/Engine/Plugins/Tests/InterchangeTests/Source/InterchangeTests/Private/InterchangeImportTestStepReimport.cpp b/Engine/Plugins/Tests/InterchangeTests/Source/InterchangeTests/Private/InterchangeImportTestStepReimport.cpp
--- a/Engine/Plugins/Tests/InterchangeTests/Source/InterchangeTests/Private/InterchangeImportTestStepReimport.cpp
+++ b/Engine/Plugins/Tests/InterchangeTests/Source/InterchangeTests/Private/InterchangeImportTestStepReimport.cpp
@@ -6,7 +6,6 @@
 
 #include UE_INLINE_GENERATED_CPP_BY_NAME(InterchangeImportTestStepReimport)
 
-
 TTuple<UE::Interchange::FAssetImportResultPtr, UE::Interchange::FSceneImportResultPtr> UInterchangeImportTestStepReimport::StartStep(FInterchangeImportTestData& Data)
 {
 	// Find the object we wish to reimport
@@ -48,10 +47,13 @@ TTuple<UE::Interchange::FAssetImportResultPtr, UE::Interchange::FSceneImportResu
 	UE::Interchange::FScopedSourceData ScopedSourceData(SourceFileToReimport.FilePath);
 
 	FImportAssetParameters Params;
-	Params.OverridePipelines.Reserve(PipelineStack.Num());
-	for (TObjectPtr<UInterchangePipelineBase> Pipeline : PipelineStack)
+	if (bUseOverridePipelineStack)
 	{
-		Params.OverridePipelines.Add(Pipeline);
+		Params.OverridePipelines.Reserve(PipelineStack.Num());
+		for (TObjectPtr<UInterchangePipelineBase> Pipeline : PipelineStack)
+		{
+			Params.OverridePipelines.Add(Pipeline);
+		}
 	}
 	Params.bIsAutomated = true;
 	Params.ReimportAsset = AssetToReimport;
@@ -85,3 +87,16 @@ FString UInterchangeImportTestStepReimport::GetContextString() const
 	return FString(TEXT("Reimporting ")) + FPaths::GetCleanFilename(SourceFileToReimport.FilePath);
 }
 
+
+bool  UInterchangeImportTestStepReimport::CanEditPipelineSettings() const
+{
+	return false;
+}
+
+void  UInterchangeImportTestStepReimport::EditPipelineSettings()
+{
+}
+
+void  UInterchangeImportTestStepReimport::ClearPipelineSettings()
+{
+}
\ No newline at end of file
diff --git a/Engine/Plugins/Tests/InterchangeTests/Source/InterchangeTests/Private/InterchangeTestPlanPipelineSettings.cpp b/Engine/Plugins/Tests/InterchangeTests/Source/InterchangeTests/Private/InterchangeTestPlanPipelineSettings.cpp
new file mode 100644
--- /dev/null
+++ b/Engine/Plugins/Tests/InterchangeTests/Source/InterchangeTests/Private/InterchangeTestPlanPipelineSettings.cpp
@@ -0,0 +1,90 @@
+// Copyright Epic Games, Inc. All Rights Reserved.
+#include "InterchangeTestPlanPipelineSettings.h"
+#include "InterchangeImportTestStepImport.h"
+#include "InterchangePipelineBase.h"
+#include "ScopedTransaction.h"
+
+void FInterchangeTestPlanPipelineSettings::UpdatePipelines(const TArray<UInterchangePipelineBase*>& InPipelines)
+{
+	if (ensure(ParentTestStep))
+	{
+		FScopedTransaction ScopedTransaction(NSLOCTEXT("InterchangeTestPlanPipelineSettings", "UpdatePipelineSettings", "Update Pipeline Settings"));
+		ParentTestStep->Modify();
+		CustomPipelines.Empty(InPipelines.Num());
+		for (const UInterchangePipelineBase* Pipeline : InPipelines)
+		{
+			CustomPipelines.Emplace(DuplicateObject<UInterchangePipelineBase>(Pipeline, ParentTestStep->GetOuter()));
+		}
+	}
+}
+
+void FInterchangeTestPlanPipelineSettings::ClearPipelines()
+{
+	if (ensure(ParentTestStep))
+	{
+		FScopedTransaction ScopedTransaction(NSLOCTEXT("InterchangeTestPlanPipelineSettings", "ClearPipelineSettings", "Clear Pipeline Settings"));
+		ParentTestStep->Modify();
+		CustomPipelines.Empty();
+	}
+}
+
+bool FInterchangeTestPlanPipelineSettings::IsUsingOverridePipelineStack() const
+{
+	if (UInterchangeImportTestStepImport* StepImport = Cast<UInterchangeImportTestStepImport>(ParentTestStep))
+	{
+		return StepImport->bUseOverridePipelineStack;
+	}
+	return false;
+}
+
+bool FInterchangeTestPlanPipelineSettings::IsUsingModifiedSettings() const
+{
+	if (UInterchangeImportTestStepImport* StepImport = Cast<UInterchangeImportTestStepImport>(ParentTestStep))
+	{
+		if (StepImport->bUseOverridePipelineStack)
+		{
+			return !StepImport->PipelineStack.IsEmpty();
+		}
+	}
+
+	return !CustomPipelines.IsEmpty();
+}
+
+bool FInterchangeTestPlanPipelineSettings::CanEditPipelineSettings() const
+{
+	if (ParentTestStep)
+	{
+		return ParentTestStep->CanEditPipelineSettings();
+	}
+
+	return false;
+}
+
+FInterchangeImportIntoLevelBoolean::FInterchangeImportIntoLevelBoolean()
+	:FInterchangeImportIntoLevelBoolean(false)
+{
+}
+
+FInterchangeImportIntoLevelBoolean::FInterchangeImportIntoLevelBoolean(bool InValue)
+	:bImportIntoLevel(InValue)
+{
+}
+
+bool FInterchangeImportIntoLevelBoolean::GetValue() const
+{
+	return bImportIntoLevel;
+}
+
+bool& FInterchangeImportIntoLevelBoolean::GetMutableValue()
+{
+	return bImportIntoLevel;
+}
+
+void FInterchangeImportIntoLevelBoolean::SetValue(bool bNewImportIntoLevel)
+{
+	if (bImportIntoLevel != bNewImportIntoLevel)
+	{
+		bImportIntoLevel = bNewImportIntoLevel;
+		OnValueChanged.ExecuteIfBound(bNewImportIntoLevel);
+	}
+}
\ No newline at end of file
diff --git a/Engine/Plugins/Tests/InterchangeTests/Source/InterchangeTests/Public/InterchangeImportTestStepBase.h b/Engine/Plugins/Tests/InterchangeTests/Source/InterchangeTests/Public/InterchangeImportTestStepBase.h
--- a/Engine/Plugins/Tests/InterchangeTests/Source/InterchangeTests/Public/InterchangeImportTestStepBase.h
+++ b/Engine/Plugins/Tests/InterchangeTests/Source/InterchangeTests/Public/InterchangeImportTestStepBase.h
@@ -33,6 +33,10 @@ public:
 	virtual FTestStepResults FinishStep(FInterchangeImportTestData& Data, FAutomationTestExecutionInfo& ExecutionInfo) PURE_VIRTUAL(UInterchangeImportTestStepBase::FinishStep, return {}; );
 	virtual FString GetContextString() const PURE_VIRTUAL(UInterchangeImportTestStepBase::GetContextString, return {}; );
 
+	virtual bool CanEditPipelineSettings() const PURE_VIRTUAL(UInterchangeImportTestStepBase::CanEditPipelineSettings, return false; );
+	virtual void EditPipelineSettings() PURE_VIRTUAL(UInterchangeImportTestStepBase::EditPipelineSettings);
+	virtual void ClearPipelineSettings() PURE_VIRTUAL(UInterchangeImportTestStepBase::ClearPipelineSettings);
+
 	bool PerformTests(FInterchangeImportTestData& Data, FAutomationTestExecutionInfo& ExecutionInfo);
 protected:
 	void SaveReloadAssets(FInterchangeImportTestData& Data);
diff --git a/Engine/Plugins/Tests/InterchangeTests/Source/InterchangeTests/Public/InterchangeImportTestStepImport.h b/Engine/Plugins/Tests/InterchangeTests/Source/InterchangeTests/Public/InterchangeImportTestStepImport.h
--- a/Engine/Plugins/Tests/InterchangeTests/Source/InterchangeTests/Public/InterchangeImportTestStepImport.h
+++ b/Engine/Plugins/Tests/InterchangeTests/Source/InterchangeTests/Public/InterchangeImportTestStepImport.h
@@ -3,9 +3,10 @@
 #pragma once
 
 #include "InterchangeImportTestStepBase.h"
-
 #include "UObject/StrongObjectPtr.h"
 
+#include "InterchangeTestPlanPipelineSettings.h"
+
 #include "InterchangeImportTestStepImport.generated.h"
 
 class UInterchangePipelineBase;
@@ -17,15 +18,23 @@ class INTERCHANGETESTS_API UInterchangeImportTestStepImport : public UInterchang
 	GENERATED_BODY()
 
 public:
+	UInterchangeImportTestStepImport(const FObjectInitializer& InObjectInitializer);
 
 	/** The source file to import (path relative to the json script) */
 	UPROPERTY(EditAnywhere, Category = General)
 	FFilePath SourceFile;
 
+	/** Whether the import should use the override pipeline stack */
+	UPROPERTY(EditAnywhere, Category = General)
+	bool bUseOverridePipelineStack = false;
+
 	/** The pipeline stack to use when importing (an empty array will use the defaults) */
-	UPROPERTY(EditAnywhere, Instanced, Category = General)
+	UPROPERTY(EditAnywhere, Instanced, Category = General, meta=(DisplayName="Override Pipeline Stack", EditCondition = "bUseOverridePipelineStack", MaxPropertyDepth = 1))
 	TArray<TObjectPtr<UInterchangePipelineBase>> PipelineStack;
 
+	UPROPERTY(EditAnywhere, Category = General)
+	FInterchangeTestPlanPipelineSettings PipelineSettings;
+
 	/** Whether the destination folder should be emptied prior to import */
 	UPROPERTY(EditAnywhere, Category = General)
 	bool bEmptyDestinationFolderPriorToImport = true;
@@ -34,10 +43,14 @@ public:
 	UPROPERTY(EditAnywhere, Category = General, Meta=(EditCondition="!bImportIntoLevel"))
 	bool bSaveThenReloadImportedAssets = true;
 
-	/** Whether we should use the import into level workflow */
-	UPROPERTY(EditAnywhere, Category = General)
+	/**  Please use the Import Into Level(v2) that has more validations. This property will soon be deprecated. */
+	UPROPERTY()
 	bool bImportIntoLevel = false;
 
+	/** Whether we should use the import into level workflow */
+	UPROPERTY(EditAnywhere, Category = General, meta=(DisplayName="Import Into Level"))
+	FInterchangeImportIntoLevelBoolean bMessageBoxImportIntoLevel;
+
 public:
 	// UInterchangeImportTestStepBase interface
 	virtual TTuple<UE::Interchange::FAssetImportResultPtr, UE::Interchange::FSceneImportResultPtr>
@@ -45,6 +58,17 @@ public:
 	virtual FTestStepResults FinishStep(FInterchangeImportTestData& Data, FAutomationTestExecutionInfo& ExecutionInfo) override;
 	virtual FString GetContextString() const override;
 
+	virtual bool CanEditPipelineSettings() const override;
+	virtual void EditPipelineSettings() override;
+	virtual void ClearPipelineSettings() override;
+
+	void OnImportIntoLevelChanged(bool bNewImportIntoLevel);
+	bool ShouldImportIntoLevelChangeRequireMessageBox() const;
+
+#if WITH_EDITOR
+	virtual void PostEditChangeProperty(FPropertyChangedEvent& PropertyChangedEvent) override;
+#endif
+
 private:
 	TStrongObjectPtr<UWorld> TransientWorld;
 };
diff --git a/Engine/Plugins/Tests/InterchangeTests/Source/InterchangeTests/Public/InterchangeImportTestStepReimport.h b/Engine/Plugins/Tests/InterchangeTests/Source/InterchangeTests/Public/InterchangeImportTestStepReimport.h
--- a/Engine/Plugins/Tests/InterchangeTests/Source/InterchangeTests/Public/InterchangeImportTestStepReimport.h
+++ b/Engine/Plugins/Tests/InterchangeTests/Source/InterchangeTests/Public/InterchangeImportTestStepReimport.h
@@ -3,7 +3,6 @@
 #pragma once
 
 #include "InterchangeImportTestStepBase.h"
-//#include "InterchangeTestFunction.h"
 #include "Templates/SubclassOf.h"
 #include "InterchangeImportTestStepReimport.generated.h"
 
@@ -16,13 +15,16 @@ class INTERCHANGETESTS_API UInterchangeImportTestStepReimport : public UIntercha
 	GENERATED_BODY()
 
 public:
-
 	/** The source file to import (path relative to the json script). */
 	UPROPERTY(EditAnywhere, Category = General)
 	FFilePath SourceFileToReimport;
 
+	/** Whether the import should use the override pipeline stack */
+	UPROPERTY(EditAnywhere, Category = General)
+	bool bUseOverridePipelineStack = false;
+
 	/** The pipeline stack to use when re-importing (an empty array will use the original import pipelines) */
-	UPROPERTY(EditAnywhere, Instanced, Category = General)
+	UPROPERTY(EditAnywhere, Instanced, Category = General, meta = (DisplayName = "Override Pipeline Stack", EditCondition = "bUseOverridePipelineStack"))
 	TArray<TObjectPtr<UInterchangePipelineBase>> PipelineStack;
 
 	/** The type of the asset to reimport. If only one such asset was imported, this is unambiguous. */
@@ -43,4 +45,9 @@ public:
 		StartStep(FInterchangeImportTestData& Data) override;
 	virtual FTestStepResults FinishStep(FInterchangeImportTestData& Data, FAutomationTestExecutionInfo& ExecutionInfo) override;
 	virtual FString GetContextString() const override;
+
+	virtual bool CanEditPipelineSettings() const override;
+	virtual void EditPipelineSettings() override;
+	virtual void ClearPipelineSettings() override;
+
 };
diff --git a/Engine/Plugins/Tests/InterchangeTests/Source/InterchangeTests/Public/InterchangeTestPlanPipelineSettings.h b/Engine/Plugins/Tests/InterchangeTests/Source/InterchangeTests/Public/InterchangeTestPlanPipelineSettings.h
new file mode 100644
--- /dev/null
+++ b/Engine/Plugins/Tests/InterchangeTests/Source/InterchangeTests/Public/InterchangeTestPlanPipelineSettings.h
@@ -0,0 +1,68 @@
+// Copyright Epic Games, Inc. All Rights Reserved.
+#pragma once
+#include "Delegates/DelegateCombinations.h"
+#include "InterchangeTestPlanPipelineSettings.generated.h"
+
+class UInterchangeImportTestStepBase;
+class UInterchangePipelineBase;
+
+USTRUCT(BlueprintType)
+struct FInterchangeTestPlanPipelineSettings
+{
+	GENERATED_BODY()
+public:
+	void UpdatePipelines(const TArray<UInterchangePipelineBase*>& InPipelines);
+	
+	void ClearPipelines();
+
+	INTERCHANGETESTS_API bool IsUsingOverridePipelineStack() const;
+
+	INTERCHANGETESTS_API bool IsUsingModifiedSettings() const;
+
+	INTERCHANGETESTS_API bool CanEditPipelineSettings() const;
+
+public:
+	/** Referenced Test Step */
+	UPROPERTY(EditAnywhere, Instanced, Category=General)
+	TArray<TObjectPtr<UInterchangePipelineBase>> CustomPipelines;
+
+	UPROPERTY()
+	TObjectPtr<UInterchangeImportTestStepBase> ParentTestStep;
+};
+
+DECLARE_DELEGATE_OneParam(FOnBoolValueChanged, bool)
+DECLARE_DELEGATE_RetVal(bool, FShouldValueChangeShowMessageBox)
+
+USTRUCT(BlueprintType)
+struct FInterchangeImportIntoLevelBoolean
+{
+	GENERATED_BODY()
+
+public:
+	FInterchangeImportIntoLevelBoolean();
+	FInterchangeImportIntoLevelBoolean(bool InValue);
+
+	INTERCHANGETESTS_API bool GetValue() const;
+	INTERCHANGETESTS_API bool& GetMutableValue();
+
+	INTERCHANGETESTS_API void SetValue(bool bNewValue);
+	
+	UPROPERTY()
+	FText MessageBoxTitle;
+
+	UPROPERTY()
+	FText MessageBoxText;
+
+	UPROPERTY()
+	FString MessageBoxShowDialogINISetting;
+
+	/** Fired off when a different value is assigned through SetValue */
+	FOnBoolValueChanged OnValueChanged;
+
+	/** Fired off in the UI when the value is about to change to check if the change requires showing message box. */
+	FShouldValueChangeShowMessageBox ShouldValueChangeShowMessageBox;
+
+private:
+	UPROPERTY()
+	bool bImportIntoLevel = false;
+};
\ No newline at end of file
diff --git a/Engine/Source/Runtime/Interchange/Engine/Private/InterchangeManager.cpp b/Engine/Source/Runtime/Interchange/Engine/Private/InterchangeManager.cpp
--- a/Engine/Source/Runtime/Interchange/Engine/Private/InterchangeManager.cpp
+++ b/Engine/Source/Runtime/Interchange/Engine/Private/InterchangeManager.cpp
@@ -314,6 +314,21 @@ UInterchangeTranslatorBase* UE::Interchange::FScopedTranslator::GetTranslator()
 	return ScopedTranslatorPtr.Get();
 }
 
+UE::Interchange::FScopedBaseNodeContainer::FScopedBaseNodeContainer()
+{
+	ScopedBaseNodeContainerPtr = TStrongObjectPtr<UInterchangeBaseNodeContainer>(NewObject<UInterchangeBaseNodeContainer>(GetTransientPackage()));
+}
+
+UE::Interchange::FScopedBaseNodeContainer::~FScopedBaseNodeContainer()
+{
+	ScopedBaseNodeContainerPtr.Reset();
+}
+
+UInterchangeBaseNodeContainer* UE::Interchange::FScopedBaseNodeContainer::GetBaseNodeContainer()
+{
+	return ScopedBaseNodeContainerPtr.Get();
+}
+
 UE::Interchange::FImportAsyncHelper::FImportAsyncHelper()
 	: AssetImportResult(MakeShared<FImportResult>())
 	, SceneImportResult(MakeShared<FImportResult>())
@@ -985,7 +1000,7 @@ UInterchangePipelineBase* UE::Interchange::GeneratePipelineInstance(const FSoftO
 	if(GeneratedPipeline)
 	{
 		// Make sure that the instance does not carry over standalone and public flags as they are not actual assets to be persisted
-		GeneratedPipeline->ClearFlags(EObjectFlags::RF_Standalone|EObjectFlags::RF_Public);
+		GeneratedPipeline->ClearFlags(EObjectFlags::RF_Standalone|EObjectFlags::RF_Public|EObjectFlags::RF_Transactional);
 	}
 
 	return GeneratedPipeline;
@@ -2479,6 +2494,20 @@ UInterchangeEditorUtilitiesBase* UInterchangeManager::GetEditorUtilities() const
 	return EditorUtilities.Get();
 }
 
+bool UInterchangeManager::CanUseTranslator(UInterchangeTranslatorBase* Translator) const
+{
+	bool bCanTranslate = true;
+	if (!Translator->IsThreadSafe())
+	{
+		const bool* NonParallelTranslatorLock = NonParallelTranslatorLocks.Find(Translator->GetClass());
+		if (NonParallelTranslatorLock && *NonParallelTranslatorLock)
+		{
+			bCanTranslate = false;
+		}
+	}
+	return bCanTranslate;
+}
+
 bool UInterchangeManager::IsObjectBeingImported(UObject* Object) const
 {
 	if (!ensure(IsInGameThread()))
diff --git a/Engine/Source/Runtime/Interchange/Engine/Public/InterchangeManager.h b/Engine/Source/Runtime/Interchange/Engine/Public/InterchangeManager.h
--- a/Engine/Source/Runtime/Interchange/Engine/Public/InterchangeManager.h
+++ b/Engine/Source/Runtime/Interchange/Engine/Public/InterchangeManager.h
@@ -97,6 +97,17 @@ namespace UE
 			TStrongObjectPtr<UInterchangeTranslatorBase> ScopedTranslatorPtr = nullptr;
 		};
 
+		class FScopedBaseNodeContainer
+		{
+		public:
+			INTERCHANGEENGINE_API FScopedBaseNodeContainer();
+			INTERCHANGEENGINE_API ~FScopedBaseNodeContainer();
+			INTERCHANGEENGINE_API UInterchangeBaseNodeContainer* GetBaseNodeContainer();
+
+		private:
+			TStrongObjectPtr<UInterchangeBaseNodeContainer> ScopedBaseNodeContainerPtr = nullptr;
+		};
+
 		enum class EImportType : uint8
 		{
 			ImportType_None,
@@ -852,6 +863,11 @@ public:
 	*/
 	INTERCHANGEENGINE_API UInterchangeEditorUtilitiesBase* GetEditorUtilities() const;
 
+	/**
+	 * Check if a non-parallel translator is unlocked for using.
+	 */
+	INTERCHANGEENGINE_API bool CanUseTranslator(UInterchangeTranslatorBase* Translator) const;
+
 protected:
 
 	/** Return true if Interchange can show UI. */
diff --git a/Engine/Source/Runtime/Interchange/Engine/Public/InterchangePipelineConfigurationBase.h b/Engine/Source/Runtime/Interchange/Engine/Public/InterchangePipelineConfigurationBase.h
--- a/Engine/Source/Runtime/Interchange/Engine/Public/InterchangePipelineConfigurationBase.h
+++ b/Engine/Source/Runtime/Interchange/Engine/Public/InterchangePipelineConfigurationBase.h
@@ -17,6 +17,7 @@ enum class EInterchangePipelineConfigurationDialogResult : uint8
 	Cancel		UMETA(DisplayName = "Cancel"),
 	Import		UMETA(DisplayName = "Import"),
 	ImportAll	UMETA(DisplayName = "Import All"),
+	SaveConfig	UMETA(DisplayName = "Save Config"),
 };
 
 USTRUCT(BlueprintType)
@@ -105,6 +106,33 @@ public:
 		return ShowReimportPipelineConfigurationDialog(PipelineStacks, OutPipelines, SourceData, Translator, BaseNodeContainer, ReimportAsset, bSceneImport);
 	}
 
+	/**
+	 * Non-virtual helper that allows Blueprint to implement an event-based function to implement ShowTestPlanPipelineConfigurationDialog().
+	 */
+	UFUNCTION(BlueprintNativeEvent, BlueprintCallable, Category = "Interchange | Translator")
+	INTERCHANGEENGINE_API EInterchangePipelineConfigurationDialogResult ScriptedShowTestPlanConfigurationDialog(TArray<FInterchangeStackInfo>& PipelineStacks
+		, TArray<UInterchangePipelineBase*>& OutPipelines
+		, UInterchangeSourceData* SourceData
+		, UInterchangeTranslatorBase* Translator
+		, UInterchangeBaseNodeContainer* BaseNodeContainer
+		, UObject* ReimportAsset
+		, bool bSceneImport = false
+		, bool bReimport = false);
+
+	/** The default implementation, which is called if the Blueprint does not have any implementation, calls the virtual ShowTestPlanAssetPipelineConfigurationDialog(). */
+	EInterchangePipelineConfigurationDialogResult ScriptedShowTestPlanConfigurationDialog_Implementation(TArray<FInterchangeStackInfo>& PipelineStacks
+		, TArray<UInterchangePipelineBase*>& OutPipelines
+		, UInterchangeSourceData* SourceData
+		, UInterchangeTranslatorBase* Translator
+		, UInterchangeBaseNodeContainer* BaseNodeContainer
+		, UObject* ReimportAsset
+		, bool bSceneImport = false
+		, bool bReimport = false)
+	{
+		//By default we call the virtual import pipeline execution
+		return ShowTestPlanConfigurationDialog(PipelineStacks, OutPipelines, SourceData, Translator, BaseNodeContainer, ReimportAsset, bSceneImport, bReimport);
+	}
+
 protected:
 
 	/**
@@ -147,4 +175,20 @@ protected:
 		//Not implemented
 		return EInterchangePipelineConfigurationDialogResult::Cancel;
 	}
+
+	/**
+	 * This function shows a dialog used to configure pipeline stacks in interchange test plan asset.
+	 */
+	virtual EInterchangePipelineConfigurationDialogResult ShowTestPlanConfigurationDialog(TArray<FInterchangeStackInfo>& PipelineStacks
+		, TArray<UInterchangePipelineBase*>& OutPipelines
+		, TWeakObjectPtr<UInterchangeSourceData> SourceData
+		, TWeakObjectPtr <UInterchangeTranslatorBase> Translator
+		, TWeakObjectPtr <UInterchangeBaseNodeContainer> BaseNodeContainer
+		, TWeakObjectPtr <UObject> ReimportAsset
+		, bool bSceneImport
+		, bool bReimport)
+	{
+		//Not implemented
+		return EInterchangePipelineConfigurationDialogResult::Cancel;
+	}
 };
