Allow tester to see and select any testplan asset under a folder
- Changed the naming convention used for the Interchange Test Plan Asset Tests.
- Object Names are used instead of the

#rnx

diff --git a/Engine/Plugins/Tests/InterchangeTests/Source/InterchangeTests/Private/InterchangeImportTest.cpp b/Engine/Plugins/Tests/InterchangeTests/Source/InterchangeTests/Private/InterchangeImportTest.cpp
--- a/Engine/Plugins/Tests/InterchangeTests/Source/InterchangeTests/Private/InterchangeImportTest.cpp
+++ b/Engine/Plugins/Tests/InterchangeTests/Source/InterchangeTests/Private/InterchangeImportTest.cpp
@@ -14,9 +14,7 @@
 #include "Modules/ModuleManager.h"
 #include "ObjectTools.h"
 
-IMPLEMENT_CUSTOM_COMPLEX_AUTOMATION_TEST(FInterchangeImportTest, FAutomationTestBase, "Editor.Interchange", EAutomationTestFlags::EditorContext | EAutomationTestFlags::EngineFilter)
-
-const FString InterchangeTestsRootGameFolder = TEXT("/Game/Tests/Interchange/");
+IMPLEMENT_CUSTOM_COMPLEX_AUTOMATION_TEST(FInterchangeImportTest, FAutomationTestBase, UE::Interchange::FInterchangeImportTestPlanStaticHelpers::GetBeautifiedTestName(), EAutomationTestFlags::EditorContext | EAutomationTestFlags::EngineFilter)
 
 void FInterchangeImportTest::GetTests(TArray<FString>& OutBeautifiedNames, TArray<FString>& OutTestCommands) const
 {
@@ -31,29 +29,23 @@ void FInterchangeImportTest::GetTests(TArray<FString>& OutBeautifiedNames, TArra
 	AssetRegistryModule.Get().GetAssetsByClass(UInterchangeImportTestPlan::StaticClass()->GetClassPathName(), AllTestPlans, true);
 
 	// Get a list of all paths containing InterchangeTestPlan assets
-	TSet<FName> Paths;
+	TSet<FString> Paths;
 	for (const FAssetData& TestPlan : AllTestPlans)
 	{
-		Paths.Add(TestPlan.PackagePath);
+		Paths.Add(TestPlan.GetObjectPathString());
 	}
 
 	// And create a sorted list from the set.
 	// Each unique path will be a sub-entry in the automated test list.
 	// All test plans within a given folder will be executed in parallel.
 
-	Paths.Sort([](const FName& A, const FName& B) { return A.ToString() < B.ToString(); });
+	Paths.Sort([](const FString& A, const FString& B) { return A < B; });
 
-	for (const FName& Path : Paths)
+	for (const FString& Path : Paths)
 	{
-		FString PathAsString = Path.ToString();
-		OutTestCommands.Add(PathAsString);
-
-		FString BeautifiedName = PathAsString;
-		FPaths::MakePathRelativeTo(BeautifiedName, *InterchangeTestsRootGameFolder);
-		BeautifiedName.ReplaceCharInline(TEXT('/'), TEXT('.'));
-		BeautifiedName.ReplaceCharInline(TEXT('\\'), TEXT('.'));
-
-		OutBeautifiedNames.Add(BeautifiedName);
+		using namespace UE::Interchange;
+		OutTestCommands.Add(Path);
+		OutBeautifiedNames.Add(FInterchangeImportTestPlanStaticHelpers::GetTestNameFromObjectPathString(Path));
 	}
 }
 
@@ -75,59 +67,18 @@ bool FInterchangeImportTest::RunTest(const FString& Path)
 	//Make sure interchange is enabled for fbx
 	CVarInterchangeFbx->Set(true, ECVF_SetByConsole);
 
-	// Determine the test plan assets within the given path which will be run in parallel
-	TArray<FInterchangeImportTestData> TestPlans;
-	if (Path.Contains(TEXT(".")))
+	
+	FInterchangeImportTestData TestPlanData;
+	FName PackageName = FName(FPaths::GetBaseFilename(Path, false));
+	FName PackagePath = FName(FPaths::GetPath(Path));
+	FName AssetName = FName(FPaths::GetBaseFilename(Path, true));
+	FTopLevelAssetPath ClassName = UInterchangeImportTestPlan::StaticClass()->GetClassPathName();
+	TestPlanData.AssetData = FAssetData(PackageName, PackagePath, AssetName, ClassName);
+	TestPlanData.TestPlan = CastChecked<UInterchangeImportTestPlan>(TestPlanData.AssetData.GetAsset());
+	
+	if (!ensure(TestPlanData.TestPlan))
 	{
-		// Run test on a single TestPlan asset
-		FName PackageName = FName(FPaths::GetBaseFilename(Path, false));
-		FName PackagePath = FName(FPaths::GetPath(Path));
-		FName AssetName = FName(FPaths::GetBaseFilename(Path, true));
-		FTopLevelAssetPath ClassName = UInterchangeImportTestPlan::StaticClass()->GetClassPathName();
-		FInterchangeImportTestData& TP = TestPlans.AddDefaulted_GetRef();
-		TP.AssetData = FAssetData(PackageName, PackagePath, AssetName, ClassName);
-		TP.TestPlan = CastChecked<UInterchangeImportTestPlan>(TP.AssetData.GetAsset());
-	}
-	else
-	{
-		// Run tests in parallel on all TestPlan assets in the given directory
-		FARFilter AssetRegistryFilter;
-		AssetRegistryFilter.PackagePaths.Add(FName(Path));
-		AssetRegistryFilter.ClassPaths.Add(UInterchangeImportTestPlan::StaticClass()->GetClassPathName());
-		AssetRegistryFilter.bRecursivePaths = false;
-
-		const FAssetRegistryModule& AssetRegistryModule = FModuleManager::LoadModuleChecked<FAssetRegistryModule>(TEXT("AssetRegistry"));
-
-		AssetRegistryModule.Get().EnumerateAssets(AssetRegistryFilter,
-			[this, &TestPlans](const FAssetData& AssetData)
-			{
-				UInterchangeImportTestPlan* AssetObject = CastChecked<UInterchangeImportTestPlan>(AssetData.GetAsset());
-
-				if (AssetObject->bIsEnabledInAutomationTests)
-				{
-					FInterchangeImportTestData& TP = TestPlans.AddDefaulted_GetRef();
-					TP.AssetData = AssetData;
-					TP.TestPlan = AssetObject;
-				}
-				else
-				{
-					FString InfoMessage;
-					if (!AssetObject->DisabledTestReason.IsEmpty())
-					{
-						InfoMessage = FString::Printf(TEXT("InterchangeImportTestPlan %s is disabled because [%s]."), *AssetObject->GetName(), *AssetObject->DisabledTestReason);
-						
-					}
-					else
-					{
-						InfoMessage = FString::Printf(TEXT("InterchangeImportTestPlan %s is disabled."), *AssetObject->GetName());
-					}
-					//Add the disable test as an informative event
-					ExecutionInfo.AddEvent(FAutomationEvent(EAutomationEventType::Info, InfoMessage));
-				}
-
-				return true;
-			}
-		);
+		return false;
 	}
 
 	// Base path to import assets to
@@ -142,7 +93,7 @@ bool FInterchangeImportTest::RunTest(const FString& Path)
 	const bool bDeleteRecursively = true;
 	IFileManager::Get().DeleteDirectory(*BaseFilePath, bRequireExists, bDeleteRecursively);
 
-	TArray<TTuple<UE::Interchange::FAssetImportResultPtr, UE::Interchange::FSceneImportResultPtr>> Results;
+	TTuple<UE::Interchange::FAssetImportResultPtr, UE::Interchange::FSceneImportResultPtr> Results;
 
 	int32 StepIndex = 0;
 	bool bFoundValidSteps = true;
@@ -150,138 +101,121 @@ bool FInterchangeImportTest::RunTest(const FString& Path)
 
 	while (bFoundValidSteps)
 	{
-		Results.Reset();
-		Results.SetNum(TestPlans.Num());
-
 		bFoundValidSteps = false;
+		TestPlanData.ImportedAssets.Empty();
 
-		// Start parallel import tasks in the interchange manager
-
-		for (int32 PlanIndex = 0; PlanIndex < TestPlans.Num(); PlanIndex++)
-		{
-			FInterchangeImportTestData& Data = TestPlans[PlanIndex];
-			Data.ImportedAssets.Empty();
-
-			FString AssetName = Data.AssetData.AssetName.ToString().Replace(TEXT("/"), TEXT("_"));
-			Data.DestAssetPackagePath = BasePackagePath / AssetName;
-			Data.DestAssetFilePath = BaseFilePath / AssetName;
+		FString TestPlanAssetName = TestPlanData.AssetData.AssetName.ToString().Replace(TEXT("/"), TEXT("_"));
+		TestPlanData.DestAssetPackagePath = BasePackagePath / TestPlanAssetName;
+		TestPlanData.DestAssetFilePath = BaseFilePath / TestPlanAssetName;
 
-			const bool bAddRecursively = true;
-			IFileManager::Get().MakeDirectory(*Data.DestAssetFilePath, bAddRecursively);
+		const bool bAddRecursively = true;
+		IFileManager::Get().MakeDirectory(*TestPlanData.DestAssetFilePath, bAddRecursively);
 
-			if (Data.TestPlan)
+		if (TestPlanData.TestPlan)
+		{
+			if (TestPlanData.TestPlan->Steps.IsValidIndex(StepIndex))
 			{
-				if (Data.TestPlan->Steps.IsValidIndex(StepIndex))
-				{
-					bFoundValidSteps = true;
-					Results[PlanIndex] = Data.TestPlan->Steps[StepIndex]->StartStep(Data);
-				}
+				bFoundValidSteps = true;
+				Results = TestPlanData.TestPlan->Steps[StepIndex]->StartStep(TestPlanData);
 			}
 		}
+		
 
 		if (bFoundValidSteps)
 		{
-			// If we started any test steps, join all the threads here, before proceeding with operations which have to be performed
-			// on the main thread, and which could potentially trigger GC.
+			UE::Interchange::FAssetImportResultPtr AssetImportResults = Results.Get<0>();
+			UE::Interchange::FSceneImportResultPtr SceneImportResults = Results.Get<1>();
 
-			for (int32 PlanIndex = 0; PlanIndex < TestPlans.Num(); PlanIndex++)
+			if (AssetImportResults)
 			{
-				if (Results[PlanIndex].Get<0>())
-				{
-					Results[PlanIndex].Get<0>()->WaitUntilDone();
-				}
+				AssetImportResults->WaitUntilDone();
+			}
 
-				if (Results[PlanIndex].Get<1>())
-				{
-					Results[PlanIndex].Get<1>()->WaitUntilDone();
-				}
+			if (SceneImportResults)
+			{
+				SceneImportResults->WaitUntilDone();
 			}
 
-			// Start post-steps
 
-			for (int32 PlanIndex = 0; PlanIndex < TestPlans.Num(); PlanIndex++)
+			// Start post-steps
+			if (TestPlanData.TestPlan)
 			{
-				FInterchangeImportTestData& Data = TestPlans[PlanIndex];
-
-				if (Data.TestPlan)
+				if (TestPlanData.TestPlan->Steps.IsValidIndex(StepIndex))
 				{
-					if (Data.TestPlan->Steps.IsValidIndex(StepIndex))
-					{
-						FString Context = Data.AssetData.AssetName.ToString() + TEXT(": ") + Data.TestPlan->Steps[StepIndex]->GetContextString();
-						ExecutionInfo.PushContext(Context);
+					FString Context = TestPlanData.AssetData.AssetName.ToString() + TEXT(": ") + TestPlanData.TestPlan->Steps[StepIndex]->GetContextString();
+					ExecutionInfo.PushContext(Context);
 
-						// Fill out list of result objects in the data object.
-						// These are the UObject* results corresponding to imported assets.
-						if (Results[PlanIndex].Get<0>())
+					// Fill out list of result objects in the data object.
+					// These are the UObject* results corresponding to imported assets.
+					if (AssetImportResults)
+					{
+						for (UObject* ImportedObject : AssetImportResults->GetImportedObjects())
 						{
-							for (UObject* ImportedObject : Results[PlanIndex].Get<0>()->GetImportedObjects())
-							{
-								Data.ResultObjects.AddUnique(ImportedObject);
-							}
+							TestPlanData.ResultObjects.AddUnique(ImportedObject);
+						}
 
-							// Also add the InterchangeResultsContainer to the data so that tests can be run on it
-							// (e.g. to check whether something imported with a specific expected error)
-							Data.InterchangeResults = Results[PlanIndex].Get<0>()->GetResults();
+						// Also add the InterchangeResultsContainer to the data so that tests can be run on it
+						// (e.g. to check whether something imported with a specific expected error)
+						TestPlanData.InterchangeResults = AssetImportResults->GetResults();
 
-							// Fill out list of imported assets as FAssetData
-							for (UObject* Object : Results[PlanIndex].Get<0>()->GetImportedObjects())
-							{
-								Data.ImportedAssets.Emplace(Object);
-							}
+						// Fill out list of imported assets as FAssetData
+						for (UObject* Object : AssetImportResults->GetImportedObjects())
+						{
+							TestPlanData.ImportedAssets.Emplace(Object);
 						}
+					}
 
-						if (Results[PlanIndex].Get<1>())
+					if (SceneImportResults)
+					{
+						for (UObject* ImportedObject : SceneImportResults->GetImportedObjects())
 						{
-							for (UObject* ImportedObject : Results[PlanIndex].Get<1>()->GetImportedObjects())
-							{
-								Data.ResultObjects.AddUnique(ImportedObject);
-							}
-
-							if (Data.InterchangeResults)
-							{
-								Data.InterchangeResults->Append(Results[PlanIndex].Get<1>()->GetResults());
-							}
-							else
-							{
-								Data.InterchangeResults = Results[PlanIndex].Get<1>()->GetResults();
-							}
+							TestPlanData.ResultObjects.AddUnique(ImportedObject);
 						}
 
-						// We don't need the Interchange results object any more, as we've already taken everything from it that we need.
-						// If we don't reset it, it will hold on to the trashed versions of the imported assets during GC.
-						Results[PlanIndex] = {};
-						//Make sure interchange result are not garbage collected in case we reload some packages
-						if (Data.InterchangeResults)
+						if (TestPlanData.InterchangeResults)
 						{
-							Data.InterchangeResults->SetFlags(RF_Standalone);
+							TestPlanData.InterchangeResults->Append(SceneImportResults->GetResults());
 						}
-						Data.TestPlan->Steps[StepIndex]->FinishStep(Data, ExecutionInfo);
+						else
+						{
+							TestPlanData.InterchangeResults = SceneImportResults->GetResults();
+						}
+					}
 
-						if (Data.InterchangeResults)
+					// We don't need the Interchange results object any more, as we've already taken everything from it that we need.
+					// If we don't reset it, it will hold on to the trashed versions of the imported assets during GC.
+					Results = {};
+					//Make sure interchange result are not garbage collected in case we reload some packages
+					if (TestPlanData.InterchangeResults)
+					{
+						TestPlanData.InterchangeResults->SetFlags(RF_Standalone);
+					}
+					TestPlanData.TestPlan->Steps[StepIndex]->FinishStep(TestPlanData, ExecutionInfo);
+
+					if (TestPlanData.InterchangeResults)
+					{
+						// Populate the automation test execution info with the interchange import results
+						for (UInterchangeResult* Result : TestPlanData.InterchangeResults->GetResults())
 						{
-							// Populate the automation test execution info with the interchange import results
-							for (UInterchangeResult* Result : Data.InterchangeResults->GetResults())
+							switch (Result->GetResultType())
 							{
-								switch (Result->GetResultType())
-								{
-								case EInterchangeResultType::Error:
-									ExecutionInfo.AddError(Result->GetText().ToString());
-									bSuccess = false;
-									break;
-
-								case EInterchangeResultType::Warning:
-									ExecutionInfo.AddWarning(Result->GetText().ToString());
-									break;
-								}
+							case EInterchangeResultType::Error:
+								ExecutionInfo.AddError(Result->GetText().ToString());
+								bSuccess = false;
+								break;
+
+							case EInterchangeResultType::Warning:
+								ExecutionInfo.AddWarning(Result->GetText().ToString());
+								break;
 							}
-							Data.InterchangeResults->ClearFlags(RF_Standalone);
 						}
+						TestPlanData.InterchangeResults->ClearFlags(RF_Standalone);
+					}
 
-						// Finished with the interchange results - null it so it will be GC'd later
-						Data.InterchangeResults = nullptr;
+					// Finished with the interchange results - null it so it will be GC'd later
+					TestPlanData.InterchangeResults = nullptr;
 
-						ExecutionInfo.PopContext();
-					}
+					ExecutionInfo.PopContext();
 				}
 			}
 
@@ -299,29 +233,24 @@ bool FInterchangeImportTest::RunTest(const FString& Path)
 	}
 
 	TArray<UObject*> ObjectsToDelete;
-	for (FInterchangeImportTestData& Data : TestPlans)
+	ObjectsToDelete.Reserve(ObjectsToDelete.Max() + TestPlanData.ResultObjects.Num());
+	for (UObject* ResultObject : TestPlanData.ResultObjects)
 	{
-		ObjectsToDelete.Reserve(ObjectsToDelete.Max() + Data.ResultObjects.Num());
-
-		for (UObject* ResultObject : Data.ResultObjects)
+		if (AActor* Actor = Cast<AActor>(ResultObject))
 		{
-			if (AActor* Actor = Cast<AActor>(ResultObject))
-			{
-				constexpr bool bShouldModifyLevel = true;
-				Actor->GetWorld()->EditorDestroyActor(Actor, bShouldModifyLevel);
-				// Call UObject::Rename directly on actor to avoid AActor::Rename which unnecessarily unregister and re-register components
-				Actor->UObject::Rename(nullptr, GetTransientPackage(), REN_DontCreateRedirectors);
-			}
-			else
+			constexpr bool bShouldModifyLevel = true;
+			Actor->GetWorld()->EditorDestroyActor(Actor, bShouldModifyLevel);
+			// Call UObject::Rename directly on actor to avoid AActor::Rename which unnecessarily unregister and re-register components
+			Actor->UObject::Rename(nullptr, GetTransientPackage(), REN_DontCreateRedirectors);
+		}
+		else
+		{
+			if (ResultObject && !ResultObject->HasAnyFlags(RF_BeginDestroyed | RF_FinishDestroyed))
 			{
-				if (ResultObject && !ResultObject->HasAnyFlags(RF_BeginDestroyed | RF_FinishDestroyed))
-				{
-					ObjectsToDelete.Add(ResultObject);
-				}
+				ObjectsToDelete.Add(ResultObject);
 			}
 		}
 	}
-
 	//Make sure all compilation is done before deleting some objects
 	FAssetCompilingManager::Get().FinishAllCompilation();
 
diff --git a/Engine/Plugins/Tests/InterchangeTests/Source/InterchangeTests/Private/InterchangeImportTestPlan.cpp b/Engine/Plugins/Tests/InterchangeTests/Source/InterchangeTests/Private/InterchangeImportTestPlan.cpp
--- a/Engine/Plugins/Tests/InterchangeTests/Source/InterchangeTests/Private/InterchangeImportTestPlan.cpp
+++ b/Engine/Plugins/Tests/InterchangeTests/Source/InterchangeTests/Private/InterchangeImportTestPlan.cpp
@@ -4,6 +4,7 @@
 #include "JsonObjectConverter.h"
 #include "Misc/AutomationTest.h"
 #include "Misc/FileHelper.h"
+#include "Misc/Paths.h"
 #include "Logging/MessageLog.h"
 
 #include UE_INLINE_GENERATED_CPP_BY_NAME(InterchangeImportTestPlan)
@@ -87,5 +88,35 @@ void UInterchangeImportTestPlan::RunThisTest()
 	}
 }
 
-#undef LOCTEXT_NAMESPACE
+namespace UE::Interchange
+{
+	FString FInterchangeImportTestPlanStaticHelpers::GetTestNameFromObjectPathString(const FString& InObjectPathString, bool bAddBeautifiedTestNamePrefix /*= false*/)
+	{
+		FString PackagePath, ObjectName;
+		InObjectPathString.Split(TEXT("."), &PackagePath, &ObjectName);
+
+		FString BeautifiedName = PackagePath;
+		FPaths::MakePathRelativeTo(BeautifiedName, *GetInterchangeImportTestRootGameFolder());
+		BeautifiedName.ReplaceCharInline(TEXT('/'), TEXT('.'));
+		BeautifiedName.ReplaceCharInline(TEXT('\\'), TEXT('.'));
+		
+		if (bAddBeautifiedTestNamePrefix)
+		{
+			return FString::Printf(TEXT("%s.%s"), *GetBeautifiedTestName() , *BeautifiedName);
+		}
+
+		return BeautifiedName;
+	}
 
+	FString FInterchangeImportTestPlanStaticHelpers::GetBeautifiedTestName()
+	{
+		return TEXT("Editor.Interchange");
+	}
+
+	FString FInterchangeImportTestPlanStaticHelpers::GetInterchangeImportTestRootGameFolder()
+	{
+		return TEXT("/Game/Tests/Interchange/");
+	}
+
+}
+#undef LOCTEXT_NAMESPACE
diff --git a/Engine/Plugins/Tests/InterchangeTests/Source/InterchangeTests/Public/InterchangeImportTestPlan.h b/Engine/Plugins/Tests/InterchangeTests/Source/InterchangeTests/Public/InterchangeImportTestPlan.h
--- a/Engine/Plugins/Tests/InterchangeTests/Source/InterchangeTests/Public/InterchangeImportTestPlan.h
+++ b/Engine/Plugins/Tests/InterchangeTests/Source/InterchangeTests/Public/InterchangeImportTestPlan.h
@@ -28,15 +28,19 @@ public:
 	UPROPERTY(EditAnywhere, Instanced, Category = Definition, Meta = (DisplayPriority = 0))
 	TArray<TObjectPtr<UInterchangeImportTestStepBase>> Steps;
 
-	/** Whether or not this test plan is currently enabled */
-	UPROPERTY(EditAnywhere, Category = Activation)
-	bool bIsEnabledInAutomationTests = true;
-
-	/** Why is the test plan disabled. */
-	UPROPERTY(EditAnywhere, Category = Activation, Meta = (EditCondition = "!bIsEnabledInAutomationTests"))
-	FString DisabledTestReason;
-
 	/** Click here to immediately run this single test through the automation framework */
 	UFUNCTION(BlueprintCallable, CallInEditor, Category = Run, Meta = (DisplayPriority = 1))
 	void RunThisTest();
 };
+
+
+namespace UE::Interchange
+{
+	class INTERCHANGETESTS_API FInterchangeImportTestPlanStaticHelpers
+	{
+	public:
+		static FString GetTestNameFromObjectPathString(const FString& InObjectPathString, bool bAddBeautifiedTestNamePrefix = false);
+		static FString GetBeautifiedTestName();
+		static FString GetInterchangeImportTestRootGameFolder();
+	};
+}
\ No newline at end of file
