Implement ability for tester to see whther a test is skipped or not from the TestPlan asset's detail panel

#rnx

diff --git a/Engine/Plugins/Tests/InterchangeTests/Source/InterchangeTestEditor/InterchangeTestEditor.Build.cs b/Engine/Plugins/Tests/InterchangeTests/Source/InterchangeTestEditor/InterchangeTestEditor.Build.cs
--- a/Engine/Plugins/Tests/InterchangeTests/Source/InterchangeTestEditor/InterchangeTestEditor.Build.cs
+++ b/Engine/Plugins/Tests/InterchangeTests/Source/InterchangeTestEditor/InterchangeTestEditor.Build.cs
@@ -32,9 +32,9 @@ namespace UnrealBuildTool.Rules
 					"AssetDefinition",
 					"AssetTools",
 					"WorkspaceMenuStructure",
-					"EditorScriptingUtilities",
 					"SlateCore",
-					"Slate"
+					"Slate",
+					"AutomationTest"
 				}
 			);
 		}
diff --git a/Engine/Plugins/Tests/InterchangeTests/Source/InterchangeTestEditor/Private/InterchangeImportTestPlanAssetDetails.cpp b/Engine/Plugins/Tests/InterchangeTests/Source/InterchangeTestEditor/Private/InterchangeImportTestPlanAssetDetails.cpp
new file mode 100644
--- /dev/null
+++ b/Engine/Plugins/Tests/InterchangeTests/Source/InterchangeTestEditor/Private/InterchangeImportTestPlanAssetDetails.cpp
@@ -0,0 +1,137 @@
+// Copyright Epic Games, Inc. All Rights Reserved.
+#include "InterchangeImportTestPlanAssetDetails.h"
+#include "InterchangeImportTestPlan.h"
+
+#include "DetailLayoutBuilder.h"
+#include "DetailWidgetRow.h"
+#include "IDetailGroup.h"
+#include "Widgets/Input/SCheckBox.h"
+#include "Widgets/Text/STextBlock.h"
+
+#include "AutomationTestExcludelist.h"
+#include "AssetRegistry/AssetData.h"
+
+#define LOCTEXT_NAMESPACE "InterchangeImportTestPlanAssetDetails"
+
+FInterchangeImportTestPlanAssetDetailsCustomization::FInterchangeImportTestPlanAssetDetailsCustomization()
+	: InterchangeImportTestPlan(nullptr)
+	, CachedDetailBuilder(nullptr)
+{
+}
+
+
+TSharedRef<IDetailCustomization> FInterchangeImportTestPlanAssetDetailsCustomization::MakeInstance()
+{
+	return MakeShareable(new FInterchangeImportTestPlanAssetDetailsCustomization);
+}
+
+
+void FInterchangeImportTestPlanAssetDetailsCustomization::CustomizeDetails(IDetailLayoutBuilder& DetailBuilder)
+{
+	CachedDetailBuilder = &DetailBuilder;
+	TArray<TWeakObjectPtr<UObject>> EditingObjects;
+	DetailBuilder.GetObjectsBeingCustomized(EditingObjects);
+	
+	if (!ensure(EditingObjects.Num() == 1))
+	{
+		return;
+	}
+
+	InterchangeImportTestPlan = Cast<UInterchangeImportTestPlan>(EditingObjects[0].Get());
+
+	if (!ensure(InterchangeImportTestPlan.IsValid()))
+	{
+		return;
+	}
+
+
+	TFunction<bool(const UInterchangeImportTestPlan* InTestPlanAsset, FString& OutReason)> IsTestPlanAssetSkipped = [](const UInterchangeImportTestPlan * InTestPlanAsset, FString & OutReason)
+	{
+		if (!InTestPlanAsset)
+		{
+			return false;
+		}
+
+		if (UAutomationTestExcludelist* AutomationTestExcludeList = UAutomationTestExcludelist::Get())
+		{
+			FAssetData TestPlanAssetData(InTestPlanAsset);
+			constexpr bool bAddBeautifiedTestNamePrefix = true;
+			const FString TestName = UE::Interchange::FInterchangeImportTestPlanStaticHelpers::GetTestNameFromObjectPathString(TestPlanAssetData.GetObjectPathString(), bAddBeautifiedTestNamePrefix);
+			if (const FAutomationTestExcludelistEntry* TestEntry = AutomationTestExcludeList->GetExcludeTestEntry(TestName))
+			{
+				OutReason = TestEntry->Reason.ToString();
+				return true;
+			}
+		}
+
+		return false;
+	};
+
+
+	IDetailCategoryBuilder& AutomationCategory = CachedDetailBuilder->EditCategory(FName("Automation"), FText::GetEmpty(), ECategoryPriority::Uncommon);
+	FDetailWidgetRow& IsTestSkippedRow = AutomationCategory.AddCustomRow(LOCTEXT("IsTestSkippedFilterString", "Is Test Skipped"));
+	IsTestSkippedRow
+		.NameContent()
+		[
+			SNew(STextBlock)
+				.Font(FAppStyle::GetFontStyle(TEXT("PropertyWindow.NormalFont")))
+				.Text(LOCTEXT("IsTestSkipped_TextBlockText", "Is Test Skipped"))
+		]
+		.ValueContent()
+		[
+			SNew(SHorizontalBox)
+			+SHorizontalBox::Slot()
+			.AutoWidth()
+			[
+				SNew(SCheckBox)
+					.IsEnabled(false)
+					.IsChecked_Lambda([this, IsTestPlanAssetSkipped]() {
+						if (InterchangeImportTestPlan.IsValid())
+						{
+							FString OutReason;
+							if (IsTestPlanAssetSkipped(InterchangeImportTestPlan.Get(), OutReason))
+							{
+								return ECheckBoxState::Checked;
+							}
+						}
+
+						return ECheckBoxState::Unchecked;
+					})
+			]
+			+SHorizontalBox::Slot()
+			.AutoWidth()
+			.Padding(2.0, 4.0)
+			[
+				SNew(STextBlock)
+				.Font(FAppStyle::GetFontStyle(TEXT("PropertyWindow.NormalFont")))
+				.Visibility_Lambda([this, IsTestPlanAssetSkipped]()
+					{
+						if (InterchangeImportTestPlan.IsValid())
+						{
+							FString OutReason;
+							if (IsTestPlanAssetSkipped(InterchangeImportTestPlan.Get(), OutReason))
+							{
+								return EVisibility::Visible;
+							}
+						}
+
+						return EVisibility::Collapsed;
+					})
+				.Text_Lambda([this, IsTestPlanAssetSkipped]()
+					{
+						if (InterchangeImportTestPlan.IsValid())
+						{
+							FString OutReason;
+							if (IsTestPlanAssetSkipped(InterchangeImportTestPlan.Get(), OutReason))
+							{
+								return FText::FromString(FString::Printf(TEXT("Reason Skipped: %s"), *OutReason));
+							}
+						}
+						return FText::GetEmpty();
+					})
+			]
+		];
+}
+
+
+#undef LOCTEXT_NAMESPACE
\ No newline at end of file
diff --git a/Engine/Plugins/Tests/InterchangeTests/Source/InterchangeTestEditor/Private/InterchangeImportTestPlanAssetDetails.h b/Engine/Plugins/Tests/InterchangeTests/Source/InterchangeTestEditor/Private/InterchangeImportTestPlanAssetDetails.h
new file mode 100644
--- /dev/null
+++ b/Engine/Plugins/Tests/InterchangeTests/Source/InterchangeTestEditor/Private/InterchangeImportTestPlanAssetDetails.h
@@ -0,0 +1,28 @@
+// Copyright Epic Games, Inc. All Rights Reserved.
+#pragma once
+
+#include "IDetailCustomization.h"
+#include "Input/Reply.h"
+#include "UObject/WeakObjectPtr.h"
+#include "Templates/SharedPointer.h"
+
+class UInterchangeImportTestPlan;
+
+class FInterchangeImportTestPlanAssetDetailsCustomization : public IDetailCustomization
+{
+public:
+	/** Makes a new instance of this detail layout class for a specific detail view requesting it */
+	static TSharedRef<IDetailCustomization> MakeInstance();
+
+	/** IDetailCustomization interface */
+	virtual void CustomizeDetails(IDetailLayoutBuilder& DetailBuilder) override;
+	/** End IDetailCustomization interface */
+
+private:
+	FInterchangeImportTestPlanAssetDetailsCustomization();
+
+private:
+	TWeakObjectPtr<UInterchangeImportTestPlan> InterchangeImportTestPlan;
+
+	IDetailLayoutBuilder* CachedDetailBuilder;	// The detail builder for this customisation
+};
\ No newline at end of file
diff --git a/Engine/Plugins/Tests/InterchangeTests/Source/InterchangeTestEditor/Private/InterchangeTestEditorModule.cpp b/Engine/Plugins/Tests/InterchangeTests/Source/InterchangeTestEditor/Private/InterchangeTestEditorModule.cpp
--- a/Engine/Plugins/Tests/InterchangeTests/Source/InterchangeTestEditor/Private/InterchangeTestEditorModule.cpp
+++ b/Engine/Plugins/Tests/InterchangeTests/Source/InterchangeTestEditor/Private/InterchangeTestEditorModule.cpp
@@ -2,9 +2,11 @@
 
 #include "InterchangeTestEditorModule.h"
 #include "InterchangeTestFunctionLayout.h"
-#include "InterchangeTestPlanPipelineSettingsLayout.h"
 #include "InterchangeTestFunction.h"
 #include "InterchangeTestPlanPipelineSettings.h"
+#include "InterchangeTestPlanPipelineSettingsLayout.h"
+#include "InterchangeImportTestPlanAssetDetails.h"
+#include "InterchangeImportTestPlan.h"
 #include "Modules/ModuleManager.h"
 #include "PropertyEditorModule.h"
 
@@ -28,15 +30,16 @@ void FInterchangeTestEditorModule::StartupModule()
 {
 	// Register the FInterchangeTestFunction struct customization
 	FPropertyEditorModule& PropertyModule = FModuleManager::LoadModuleChecked<FPropertyEditorModule>("PropertyEditor");
+
+	PropertyModule.RegisterCustomClassLayout(UInterchangeImportTestPlan::StaticClass()->GetFName(), 
+		FOnGetDetailCustomizationInstance::CreateStatic(&FInterchangeImportTestPlanAssetDetailsCustomization::MakeInstance));
+
 	PropertyModule.RegisterCustomPropertyTypeLayout(FInterchangeTestFunction::StaticStruct()->GetFName(),
 	FOnGetPropertyTypeCustomizationInstance::CreateStatic(&FInterchangeTestFunctionLayout::MakeInstance));
 
 	PropertyModule.RegisterCustomPropertyTypeLayout(FInterchangeTestPlanPipelineSettings::StaticStruct()->GetFName(),
 	FOnGetPropertyTypeCustomizationInstance::CreateStatic(&FInterchangeTestPlanPipelineSettingsLayout::MakeInstance));
 
-	PropertyModule.RegisterCustomPropertyTypeLayout(FInterchangeImportIntoLevelBoolean::StaticStruct()->GetFName(),
-	FOnGetPropertyTypeCustomizationInstance::CreateStatic(&FInterchangeTestPlanImportIntoLevelLayout::MakeInstance));
-
 	PropertyModule.NotifyCustomizationModuleChanged();
 }
 
@@ -46,9 +49,9 @@ void FInterchangeTestEditorModule::ShutdownModule()
 	if (UObjectInitialized())
 	{
 		FPropertyEditorModule& PropertyModule = FModuleManager::LoadModuleChecked<FPropertyEditorModule>("PropertyEditor");
+		PropertyModule.UnregisterCustomClassLayout(UInterchangeImportTestPlan::StaticClass()->GetFName());
 		PropertyModule.UnregisterCustomPropertyTypeLayout(FInterchangeTestFunction::StaticStruct()->GetFName());
 		PropertyModule.UnregisterCustomPropertyTypeLayout(FInterchangeTestPlanPipelineSettings::StaticStruct()->GetFName());
-		PropertyModule.UnregisterCustomPropertyTypeLayout(FInterchangeImportIntoLevelBoolean::StaticStruct()->GetFName());
 	}
 }
 
diff --git a/Engine/Plugins/Tests/InterchangeTests/Source/InterchangeTestEditor/Private/InterchangeTestPlanPipelineSettingsLayout.cpp b/Engine/Plugins/Tests/InterchangeTests/Source/InterchangeTestEditor/Private/InterchangeTestPlanPipelineSettingsLayout.cpp
--- a/Engine/Plugins/Tests/InterchangeTests/Source/InterchangeTestEditor/Private/InterchangeTestPlanPipelineSettingsLayout.cpp
+++ b/Engine/Plugins/Tests/InterchangeTests/Source/InterchangeTestEditor/Private/InterchangeTestPlanPipelineSettingsLayout.cpp
@@ -7,7 +7,6 @@
 
 #include "DetailWidgetRow.h"
 #include "DetailLayoutBuilder.h"
-#include "Dialogs/Dialogs.h"
 #include "Editor.h"
 #include "HAL/IConsoleManager.h"
 #include "IDetailChildrenBuilder.h"
@@ -26,7 +25,6 @@ static FAutoConsoleVariableRef CCvarInterchangeDefaultShowEssentialsView(
 	ECVF_Default);
 
 
-
 /////////// FInterchangeTestPlanPipelineSettingsLayout ///////////
 
 TSharedRef<IPropertyTypeCustomization> FInterchangeTestPlanPipelineSettingsLayout::MakeInstance()
@@ -160,110 +158,4 @@ FInterchangeTestPlanPipelineSettings* FInterchangeTestPlanPipelineSettingsLayout
 	return static_cast<FInterchangeTestPlanPipelineSettings*>(StructPtr);
 }
 
-
-/////////// FInterchangeTestPlanImportIntoLevelLayout ///////////
-TSharedRef<IPropertyTypeCustomization> FInterchangeTestPlanImportIntoLevelLayout::MakeInstance()
-{
-	return MakeShared<FInterchangeTestPlanImportIntoLevelLayout>();
-}
-
-FInterchangeTestPlanImportIntoLevelLayout::FInterchangeTestPlanImportIntoLevelLayout()
-{
-}
-
-FInterchangeTestPlanImportIntoLevelLayout::~FInterchangeTestPlanImportIntoLevelLayout()
-{
-	FEditorDelegates::PostUndoRedo.RemoveAll(this);
-}
-
-void FInterchangeTestPlanImportIntoLevelLayout::CustomizeHeader(TSharedRef<IPropertyHandle> StructPropertyHandle, FDetailWidgetRow& HeaderRow, IPropertyTypeCustomizationUtils& StructCustomizationUtils)
-{
-	FEditorDelegates::PostUndoRedo.AddSP(this, &FInterchangeTestPlanImportIntoLevelLayout::RefreshLayout);
-
-	StructProperty = StructPropertyHandle;
-	PropertyUtilities = StructCustomizationUtils.GetPropertyUtilities();
-
-	HeaderRow
-		.NameContent()
-		[
-			StructPropertyHandle->CreatePropertyNameWidget()
-		]
-		.ValueContent()
-		[
-			SNew(SCheckBox)
-			.IsChecked_Lambda([this]()
-				{
-					if (FInterchangeImportIntoLevelBoolean* MessageBoxBoolean = GetStruct())
-					{
-						return MessageBoxBoolean->GetValue() ? ECheckBoxState::Checked : ECheckBoxState::Unchecked;
-					}
-					return ECheckBoxState::Unchecked;
-				})
-			.OnCheckStateChanged_Lambda([this](ECheckBoxState NewState)
-				{
-					if (FInterchangeImportIntoLevelBoolean* MessageBoxBoolean = GetStruct())
-					{
-						const bool bShouldShowMessageBox = MessageBoxBoolean->ShouldValueChangeShowMessageBox.IsBound() && MessageBoxBoolean->ShouldValueChangeShowMessageBox.Execute();
-						if (bShouldShowMessageBox)
-						{
-							FSuppressableWarningDialog::FSetupInfo DialogSetupInfo = FSuppressableWarningDialog::FSetupInfo(
-								MessageBoxBoolean->MessageBoxText,
-								MessageBoxBoolean->MessageBoxTitle,
-								MessageBoxBoolean->MessageBoxShowDialogINISetting);
-
-							DialogSetupInfo.ConfirmText = LOCTEXT("TestPlanMessageBoxBoolean_Confirm", "Yes");
-							DialogSetupInfo.CancelText = LOCTEXT("TestPlanMessageBoxBoolean_Cancel", "No");
-							DialogSetupInfo.bDefaultToSuppressInTheFuture = false;
-							DialogSetupInfo.CheckBoxText = LOCTEXT("TestPlanMessageBoxBoolean_CheckBoxText", "Don't show this dialog again");
-							FSuppressableWarningDialog SuppressableWarningDialog(DialogSetupInfo);
-							FSuppressableWarningDialog::EResult DialogResult = SuppressableWarningDialog.ShowModal();
-
-							if (DialogResult != FSuppressableWarningDialog::EResult::Cancel)
-							{
-								MessageBoxBoolean->SetValue(NewState == ECheckBoxState::Checked);
-							}
-						}
-						else 
-						{
-							MessageBoxBoolean->SetValue(NewState == ECheckBoxState::Checked);
-						}
-					}
-				})
-		];
-}
-
-void FInterchangeTestPlanImportIntoLevelLayout::CustomizeChildren(TSharedRef<IPropertyHandle> StructPropertyHandle, IDetailChildrenBuilder& StructBuilder, IPropertyTypeCustomizationUtils& StructCustomizationUtils)
-{
-	if (!StructPropertyHandle->IsValidHandle())
-	{
-		return;
-	}
-
-	uint32 NumChildren = 0;
-	StructPropertyHandle->GetNumChildren(NumChildren);
-
-	for (uint32 ChildIndex = 0; ChildIndex < NumChildren; ChildIndex++)
-	{
-		StructBuilder.AddProperty(StructPropertyHandle->GetChildHandle(ChildIndex).ToSharedRef());
-	}
-}
-
-void FInterchangeTestPlanImportIntoLevelLayout::RefreshLayout()
-{
-	if (PropertyUtilities.IsValid())
-	{
-		PropertyUtilities->ForceRefresh();
-	}
-}
-
-FInterchangeImportIntoLevelBoolean* FInterchangeTestPlanImportIntoLevelLayout::GetStruct() const
-{
-	// Get address of the FInterchangeTestPlanImportIntoLevelLayout struct being viewed.
-	// We only ever expect the property handle to be linked to a single instance.
-
-	void* StructPtr = nullptr;
-	StructProperty->GetValueData(StructPtr);
-	return static_cast<FInterchangeImportIntoLevelBoolean*>(StructPtr);
-}
-
 #undef LOCTEXT_NAMESPACE 
\ No newline at end of file
diff --git a/Engine/Plugins/Tests/InterchangeTests/Source/InterchangeTestEditor/Private/InterchangeTestPlanPipelineSettingsLayout.h b/Engine/Plugins/Tests/InterchangeTests/Source/InterchangeTestEditor/Private/InterchangeTestPlanPipelineSettingsLayout.h
--- a/Engine/Plugins/Tests/InterchangeTests/Source/InterchangeTestEditor/Private/InterchangeTestPlanPipelineSettingsLayout.h
+++ b/Engine/Plugins/Tests/InterchangeTests/Source/InterchangeTestEditor/Private/InterchangeTestPlanPipelineSettingsLayout.h
@@ -6,7 +6,6 @@ class FDetailWidgetRow;
 class IDetailChildrenBuilder;
 class IPropertyUtilities;
 struct FInterchangeTestPlanPipelineSettings;
-struct FInterchangeImportIntoLevelBoolean;
 class FReply;
 
 class FInterchangeTestPlanPipelineSettingsLayout : public IPropertyTypeCustomization
@@ -35,28 +34,4 @@ private:
 	TSharedPtr<IPropertyUtilities> PropertyUtilities;
 	TSharedPtr<IPropertyHandle> StructProperty;
 
-};
-
-class FInterchangeTestPlanImportIntoLevelLayout : public IPropertyTypeCustomization
-{
-public:
-	/** Makes a new instance of this layout class for a specific detail view requesting it */
-	static TSharedRef<IPropertyTypeCustomization> MakeInstance();
-
-	FInterchangeTestPlanImportIntoLevelLayout();
-	virtual ~FInterchangeTestPlanImportIntoLevelLayout();
-
-	// IPropertyTypeCustomization interface
-	virtual void CustomizeHeader(TSharedRef<IPropertyHandle> StructPropertyHandle, FDetailWidgetRow& HeaderRow, IPropertyTypeCustomizationUtils& StructCustomizationUtils) override;
-	virtual void CustomizeChildren(TSharedRef<IPropertyHandle> StructPropertyHandle, IDetailChildrenBuilder& StructBuilder, IPropertyTypeCustomizationUtils& StructCustomizationUtils) override;
-
-private:
-	void RefreshLayout();
-private:
-	FInterchangeImportIntoLevelBoolean* GetStruct() const;
-
-private:
-	TSharedPtr<IPropertyUtilities> PropertyUtilities;
-	TSharedPtr<IPropertyHandle> StructProperty;
-
 };
\ No newline at end of file
diff --git a/Engine/Plugins/Tests/InterchangeTests/Source/InterchangeTests/Private/InterchangeImportTestStepImport.cpp b/Engine/Plugins/Tests/InterchangeTests/Source/InterchangeTests/Private/InterchangeImportTestStepImport.cpp
--- a/Engine/Plugins/Tests/InterchangeTests/Source/InterchangeTests/Private/InterchangeImportTestStepImport.cpp
+++ b/Engine/Plugins/Tests/InterchangeTests/Source/InterchangeTests/Private/InterchangeImportTestStepImport.cpp
@@ -24,13 +24,6 @@
 UInterchangeImportTestStepImport::UInterchangeImportTestStepImport(const FObjectInitializer& InObjectInitializer)
 {
 	PipelineSettings.ParentTestStep = this;
-
-	bMessageBoxImportIntoLevel = FInterchangeImportIntoLevelBoolean(bImportIntoLevel);
-	bMessageBoxImportIntoLevel.MessageBoxTitle = LOCTEXT("TestStepImport_ImportIntoLevelTitle", "Import Type Changed");
-	bMessageBoxImportIntoLevel.MessageBoxText = LOCTEXT("TestStepImport_ImportIntoLevelText", "Import type is changed. This will delete all the modifications made.\nNOTE: This change is irreversible.\nDo you still wish to continue? ");
-	bMessageBoxImportIntoLevel.MessageBoxShowDialogINISetting = "InterchangeImportTestPlanImportIntoLevelWarning";
-	bMessageBoxImportIntoLevel.OnValueChanged.BindUObject(this, &UInterchangeImportTestStepImport::OnImportIntoLevelChanged);
-	bMessageBoxImportIntoLevel.ShouldValueChangeShowMessageBox.BindUObject(this, &UInterchangeImportTestStepImport::ShouldImportIntoLevelChangeRequireMessageBox);
 }
 
 TTuple<UE::Interchange::FAssetImportResultPtr, UE::Interchange::FSceneImportResultPtr> UInterchangeImportTestStepImport::StartStep(FInterchangeImportTestData& Data)
@@ -373,7 +366,6 @@ void UInterchangeImportTestStepImport::EditPipelineSettings()
 void UInterchangeImportTestStepImport::ClearPipelineSettings()
 {
 #if WITH_EDITOR
-	
 	if (bUseOverridePipelineStack)
 	{
 		FScopedTransaction ScopedTransaction(LOCTEXT("ClearPipelineSettings", "Clear Pipeline Settings"));
@@ -387,12 +379,6 @@ void UInterchangeImportTestStepImport::ClearPipelineSettings()
 #endif
 }
 
-void UInterchangeImportTestStepImport::OnImportIntoLevelChanged(bool bNewImportIntoLevel)
-{
-	bImportIntoLevel = bNewImportIntoLevel;
-	PipelineSettings.ClearPipelines();
-}
-
 bool UInterchangeImportTestStepImport::ShouldImportIntoLevelChangeRequireMessageBox() const
 {
 	return !bUseOverridePipelineStack && PipelineSettings.IsUsingModifiedSettings();
@@ -406,19 +392,45 @@ void UInterchangeImportTestStepImport::PostEditChangeProperty(FPropertyChangedEv
 		if (bUseOverridePipelineStack && !PipelineSettings.CustomPipelines.IsEmpty())
 		{
 			FSuppressableWarningDialog::FSetupInfo DialogSetupInfo = FSuppressableWarningDialog::FSetupInfo(
-				LOCTEXT("TestStepImport_UseOverridePipelinesText", "You are using override pipelines now. This will erase all the previous modifications made to the default pipeline settings."),
-				LOCTEXT("TestStepImport_UseOverridePipelinesTitle", "Using Override Pipelines"),
+				LOCTEXT("UseOverridePipelinesDialogText", "You are using override pipelines now. This will erase all the previous modifications made to the default pipeline settings."),
+				LOCTEXT("UseOverridePipelinesDialogTitle", "Using Override Pipelines"),
 				"InterchangeImportTestPlanUseOverridePipelinesWarning");
 
-			DialogSetupInfo.ConfirmText = LOCTEXT("TestStepImport_UseOverridePipelines_Confirm", "OK");
+			DialogSetupInfo.ConfirmText = LOCTEXT("UseOverridePipelinesDialogOptionConfirm", "OK");
 			DialogSetupInfo.bDefaultToSuppressInTheFuture = false;
-			DialogSetupInfo.CheckBoxText = LOCTEXT("TestStepImport_UseOverridePipelines_CheckBoxText", "Don't show this dialog again");
+			DialogSetupInfo.CheckBoxText = LOCTEXT("UseOverridePipelinesDialogCheckBoxText", "Don't show this dialog again");
 			FSuppressableWarningDialog SuppressableWarningDialog(DialogSetupInfo);
 			SuppressableWarningDialog.ShowModal();
 
 			PipelineSettings.ClearPipelines();
 		}
 	}
+
+	if (PropertyChangedEvent.Property->GetFName() == GET_MEMBER_NAME_CHECKED(UInterchangeImportTestStepImport, bImportIntoLevel))
+	{
+		if (ShouldImportIntoLevelChangeRequireMessageBox())
+		{
+			FSuppressableWarningDialog::FSetupInfo DialogSetupInfo = FSuppressableWarningDialog::FSetupInfo(
+				LOCTEXT("ImportIntoLevelDialogMessage", "Import type is changed. This will delete all the modifications made.\nNOTE: This change is irreversible.\nDo you still wish to continue? "),
+				LOCTEXT("ImportIntoLevelDialogTitle", "Import Type Changed"),
+				"InterchangeImportTestPlanImportIntoLevelWarning");
+
+			DialogSetupInfo.ConfirmText = LOCTEXT("ImportIntoLevelDialogOptionConfirm", "Yes");
+			DialogSetupInfo.CancelText =  LOCTEXT("ImportIntoLevelDialogOptionCancel", "No");
+			DialogSetupInfo.bDefaultToSuppressInTheFuture = false;
+			DialogSetupInfo.CheckBoxText = LOCTEXT("ImportIntoLevelDialogCheckBoxText", "Don't show this dialog again");
+			FSuppressableWarningDialog SuppressableWarningDialog(DialogSetupInfo);
+			FSuppressableWarningDialog::EResult DialogResult = SuppressableWarningDialog.ShowModal();
+			if (DialogResult != FSuppressableWarningDialog::Cancel)
+			{
+				PipelineSettings.ClearPipelines();
+			}
+			else
+			{
+				bImportIntoLevel = !bImportIntoLevel;
+			}
+		}
+	}
 }
 #endif
 
diff --git a/Engine/Plugins/Tests/InterchangeTests/Source/InterchangeTests/Private/InterchangeTestPlanPipelineSettings.cpp b/Engine/Plugins/Tests/InterchangeTests/Source/InterchangeTests/Private/InterchangeTestPlanPipelineSettings.cpp
--- a/Engine/Plugins/Tests/InterchangeTests/Source/InterchangeTests/Private/InterchangeTestPlanPipelineSettings.cpp
+++ b/Engine/Plugins/Tests/InterchangeTests/Source/InterchangeTests/Private/InterchangeTestPlanPipelineSettings.cpp
@@ -58,33 +58,4 @@ bool FInterchangeTestPlanPipelineSettings::CanEditPipelineSettings() const
 	}
 
 	return false;
-}
-
-FInterchangeImportIntoLevelBoolean::FInterchangeImportIntoLevelBoolean()
-	:FInterchangeImportIntoLevelBoolean(false)
-{
-}
-
-FInterchangeImportIntoLevelBoolean::FInterchangeImportIntoLevelBoolean(bool InValue)
-	:bImportIntoLevel(InValue)
-{
-}
-
-bool FInterchangeImportIntoLevelBoolean::GetValue() const
-{
-	return bImportIntoLevel;
-}
-
-bool& FInterchangeImportIntoLevelBoolean::GetMutableValue()
-{
-	return bImportIntoLevel;
-}
-
-void FInterchangeImportIntoLevelBoolean::SetValue(bool bNewImportIntoLevel)
-{
-	if (bImportIntoLevel != bNewImportIntoLevel)
-	{
-		bImportIntoLevel = bNewImportIntoLevel;
-		OnValueChanged.ExecuteIfBound(bNewImportIntoLevel);
-	}
 }
\ No newline at end of file
diff --git a/Engine/Plugins/Tests/InterchangeTests/Source/InterchangeTests/Public/InterchangeImportTestStepImport.h b/Engine/Plugins/Tests/InterchangeTests/Source/InterchangeTests/Public/InterchangeImportTestStepImport.h
--- a/Engine/Plugins/Tests/InterchangeTests/Source/InterchangeTests/Public/InterchangeImportTestStepImport.h
+++ b/Engine/Plugins/Tests/InterchangeTests/Source/InterchangeTests/Public/InterchangeImportTestStepImport.h
@@ -4,7 +4,6 @@
 
 #include "InterchangeImportTestStepBase.h"
 #include "UObject/StrongObjectPtr.h"
-
 #include "InterchangeTestPlanPipelineSettings.h"
 
 #include "InterchangeImportTestStepImport.generated.h"
@@ -40,17 +39,13 @@ public:
 	bool bEmptyDestinationFolderPriorToImport = true;
 
 	/** Whether imported assets should be saved and freshly reloaded after import */
-	UPROPERTY(EditAnywhere, Category = General, Meta=(EditCondition="!bImportIntoLevel"))
+	UPROPERTY(EditAnywhere, Category = General, meta=(EditCondition="!bImportIntoLevel"))
 	bool bSaveThenReloadImportedAssets = true;
 
-	/**  Please use the Import Into Level(v2) that has more validations. This property will soon be deprecated. */
-	UPROPERTY()
+	/**  Whether we should use the import into level workflow */
+	UPROPERTY(EditAnywhere, Category = General)
 	bool bImportIntoLevel = false;
 
-	/** Whether we should use the import into level workflow */
-	UPROPERTY(EditAnywhere, Category = General, meta=(DisplayName="Import Into Level"))
-	FInterchangeImportIntoLevelBoolean bMessageBoxImportIntoLevel;
-
 public:
 	// UInterchangeImportTestStepBase interface
 	virtual TTuple<UE::Interchange::FAssetImportResultPtr, UE::Interchange::FSceneImportResultPtr>
@@ -62,7 +57,6 @@ public:
 	virtual void EditPipelineSettings() override;
 	virtual void ClearPipelineSettings() override;
 
-	void OnImportIntoLevelChanged(bool bNewImportIntoLevel);
 	bool ShouldImportIntoLevelChangeRequireMessageBox() const;
 
 #if WITH_EDITOR
diff --git a/Engine/Plugins/Tests/InterchangeTests/Source/InterchangeTests/Public/InterchangeTestPlanPipelineSettings.h b/Engine/Plugins/Tests/InterchangeTests/Source/InterchangeTests/Public/InterchangeTestPlanPipelineSettings.h
--- a/Engine/Plugins/Tests/InterchangeTests/Source/InterchangeTests/Public/InterchangeTestPlanPipelineSettings.h
+++ b/Engine/Plugins/Tests/InterchangeTests/Source/InterchangeTests/Public/InterchangeTestPlanPipelineSettings.h
@@ -28,41 +28,4 @@ public:
 
 	UPROPERTY()
 	TObjectPtr<UInterchangeImportTestStepBase> ParentTestStep;
-};
-
-DECLARE_DELEGATE_OneParam(FOnBoolValueChanged, bool)
-DECLARE_DELEGATE_RetVal(bool, FShouldValueChangeShowMessageBox)
-
-USTRUCT(BlueprintType)
-struct FInterchangeImportIntoLevelBoolean
-{
-	GENERATED_BODY()
-
-public:
-	FInterchangeImportIntoLevelBoolean();
-	FInterchangeImportIntoLevelBoolean(bool InValue);
-
-	INTERCHANGETESTS_API bool GetValue() const;
-	INTERCHANGETESTS_API bool& GetMutableValue();
-
-	INTERCHANGETESTS_API void SetValue(bool bNewValue);
-	
-	UPROPERTY()
-	FText MessageBoxTitle;
-
-	UPROPERTY()
-	FText MessageBoxText;
-
-	UPROPERTY()
-	FString MessageBoxShowDialogINISetting;
-
-	/** Fired off when a different value is assigned through SetValue */
-	FOnBoolValueChanged OnValueChanged;
-
-	/** Fired off in the UI when the value is about to change to check if the change requires showing message box. */
-	FShouldValueChangeShowMessageBox ShouldValueChangeShowMessageBox;
-
-private:
-	UPROPERTY()
-	bool bImportIntoLevel = false;
 };
\ No newline at end of file
