[Interchange Test Plan] Implement the ability for a tester to execute an image comparison when Import Into Level is selected
- Provides a way to optionally set a world in which the assets would be imported allowing to take a screenshot. Otherwise no screenshot would be taken as the import uses a Transient World.
- Other options such as Wireframe Opacity, View Mode and Comparison Tolerances are exposed in the details panel that would affect the screenshots and the comparison.
- The screenshot automatically tries to focus on all actors imported, but could be further filtered by classes through an array of TSubClassOf<AActor> exposed in the details panel.

#rnx

#virtualized

diff --git a/Engine/Plugins/Tests/InterchangeTests/Source/InterchangeTests/InterchangeTests.Build.cs b/Engine/Plugins/Tests/InterchangeTests/Source/InterchangeTests/InterchangeTests.Build.cs
--- a/Engine/Plugins/Tests/InterchangeTests/Source/InterchangeTests/InterchangeTests.Build.cs
+++ b/Engine/Plugins/Tests/InterchangeTests/Source/InterchangeTests/InterchangeTests.Build.cs
@@ -33,7 +33,9 @@ namespace UnrealBuildTool.Rules
 					"StaticMeshDescription",
 					"SkeletalMeshDescription",
 					"UnrealEd",
-					"VariantManagerContent"
+					"VariantManagerContent",
+					"FunctionalTesting",
+					"LevelEditor"
 				}
 			);
 		}
diff --git a/Engine/Plugins/Tests/InterchangeTests/Source/InterchangeTests/Private/InterchangeAutomatedTestUtils.cpp b/Engine/Plugins/Tests/InterchangeTests/Source/InterchangeTests/Private/InterchangeAutomatedTestUtils.cpp
new file mode 100644
--- /dev/null
+++ b/Engine/Plugins/Tests/InterchangeTests/Source/InterchangeTests/Private/InterchangeAutomatedTestUtils.cpp
@@ -0,0 +1,279 @@
+// Copyright Epic Games, Inc. All Rights Reserved.
+#include "InterchangeAutomatedTestUtils.h"
+#include "Editor.h"
+#include "Editor/Transactor.h"
+
+#include "AutomationBlueprintFunctionLibrary.h"
+#include "Misc/AutomationTest.h"
+#include "AssetCompilingManager.h"
+#include "AssetRegistry/ARFilter.h"
+#include "AssetRegistry/AssetRegistryModule.h"
+#include "HAL/FileManager.h"
+#include "InterchangeHelper.h"
+#include "InterchangeImportTestData.h"
+#include "InterchangeImportTestPlan.h"
+#include "InterchangeImportTestStepBase.h"
+#include "Modules/ModuleManager.h"
+#include "ObjectTools.h"
+#include "Subsystems/UnrealEditorSubsystem.h"
+
+bool FInterchangeIntializeStepCommand::Update()
+{
+	TestStepData->StepIndex = StepIndex;
+	TestStepData->TestPlanData.ImportedAssets.Empty();
+	TestStepData->Results = TestStepData->TestPlanData.TestPlan->Steps[TestStepData->StepIndex]->StartStep(TestStepData->TestPlanData);
+	if (UE::Interchange::FAssetImportResultPtr AssetImportResults = TestStepData->Results.Get<0>())
+	{
+		AssetImportResults->WaitUntilDone();
+	}
+
+	if (UE::Interchange::FSceneImportResultPtr SceneImportResults = TestStepData->Results.Get<1>())
+	{
+		SceneImportResults->WaitUntilDone();
+	}
+
+	return true;
+}
+
+bool FInterchangeInterStepCollectResultCommand::Update()
+{
+	if (FAutomationTestBase* CurrentTest = FAutomationTestFramework::Get().GetCurrentTest())
+	{
+		FString Context = TestStepData->TestPlanData.AssetData.AssetName.ToString() + TEXT(": ") + TestStepData->TestPlanData.TestPlan->Steps[TestStepData->StepIndex]->GetContextString();
+		CurrentTest->PushContext(Context);
+
+		// Fill out list of result objects in the data object.
+		// These are the UObject* results corresponding to imported assets.
+		if (UE::Interchange::FAssetImportResultPtr AssetImportResults = TestStepData->Results.Get<0>())
+		{
+			for (UObject* ImportedObject : AssetImportResults->GetImportedObjects())
+			{
+				TestStepData->TestPlanData.ResultObjects.AddUnique(ImportedObject);
+			}
+
+			// Also add the InterchangeResultsContainer to the data so that tests can be run on it
+			// (e.g. to check whether something imported with a specific expected error)
+			TestStepData->TestPlanData.InterchangeResults = AssetImportResults->GetResults();
+
+			// Fill out list of imported assets as FAssetData
+			for (UObject* Object : AssetImportResults->GetImportedObjects())
+			{
+				TestStepData->TestPlanData.ImportedAssets.Emplace(Object);
+			}
+		}
+
+		if (UE::Interchange::FSceneImportResultPtr SceneImportResults = TestStepData->Results.Get<1>())
+		{
+			for (UObject* ImportedObject : SceneImportResults->GetImportedObjects())
+			{
+				TestStepData->TestPlanData.ResultObjects.AddUnique(ImportedObject);
+			}
+
+			if (TestStepData->TestPlanData.InterchangeResults)
+			{
+				TestStepData->TestPlanData.InterchangeResults->Append(SceneImportResults->GetResults());
+			}
+			else
+			{
+				TestStepData->TestPlanData.InterchangeResults = SceneImportResults->GetResults();
+			}
+		}
+
+		// We don't need the Interchange results object any more, as we've already taken everything from it that we need.
+		// If we don't reset it, it will hold on to the trashed versions of the imported assets during GC.
+		TestStepData->Results = {};
+		//Make sure interchange result are not garbage collected in case we reload some packages
+		if (TestStepData->TestPlanData.InterchangeResults)
+		{
+			TestStepData->TestPlanData.InterchangeResults->SetFlags(RF_Standalone);
+		}
+	}
+	return true;
+}
+
+bool FInterchangeSetupScreenshotViewportCommand::Update()
+{
+	TArray<AActor*> ActorsInLevel;
+	const bool bUseActorFilter = !ActorClassFilter.IsEmpty();
+	for (UObject* Object : TestStepData->TestPlanData.ResultObjects)
+	{
+		if (bUseActorFilter)
+		{
+			if (ActorClassFilter.Contains(Object->GetClass()))
+			{
+				ActorsInLevel.Add(Cast<AActor>(Object));
+			}
+		}
+		else if (Object->IsA<AActor>())
+		{
+			ActorsInLevel.Add(Cast<AActor>(Object));
+		}
+	}
+
+	if (!ActorsInLevel.IsEmpty())
+	{
+		// Cache current view mode and wireframe opacity to be restored later during cleanup.
+		TestStepData->CachedScreenshotParameters.WireframeOpacity = UAutomationBlueprintFunctionLibrary::GetEditorActiveViewportWireframeOpacity();
+		TestStepData->CachedScreenshotParameters.ViewMode = UAutomationBlueprintFunctionLibrary::GetEditorActiveViewportViewMode();
+
+		UAutomationBlueprintFunctionLibrary::SetEditorActiveViewportViewMode(ScreenshotParameters.ViewMode);
+		UAutomationBlueprintFunctionLibrary::SetEditorActiveViewportWireframeOpacity(ScreenshotParameters.WireframeOpacity);
+
+		// Reset Viewport Camera Position to Origin for consistent results.
+		if (UUnrealEditorSubsystem* UnrealEditorSubsystem = GEditor->GetEditorSubsystem<UUnrealEditorSubsystem>())
+		{
+			UnrealEditorSubsystem->SetLevelViewportCameraInfo(FVector::ZeroVector, FRotator::ZeroRotator);
+		}
+
+		constexpr bool bActiveViewportOnly = true;
+		GEditor->MoveViewportCamerasToActor(ActorsInLevel, bActiveViewportOnly);
+	}
+	else
+	{
+		if (FAutomationTestBase* CurrentTest = FAutomationTestFramework::Get().GetCurrentTest())
+		{
+			if (bUseActorFilter)
+			{
+				const TArray<UClass*> ActorClasses = ActorClassFilter.Array();
+				FString ClassNamesString = TEXT("");
+				for (int32 ActorClassIndex = 0; ActorClassIndex < ActorClasses.Num(); ++ActorClassIndex)
+				{
+					ClassNamesString += ActorClasses[ActorClassIndex]->GetName();
+					if (ActorClassIndex < ActorClasses.Num() - 1)
+					{
+						ClassNamesString += TEXT(", ");
+					}
+				}
+				CurrentTest->AddWarning(FString::Printf(TEXT("Current scene does not contain any actors of classes: [%s]"), *ClassNamesString));
+			}
+			else
+			{
+				CurrentTest->AddWarning(TEXT("Current scene does not contain any actors"));
+			}
+		}
+	}
+	return true;
+}
+
+bool FInterchangeCaptureScreenshotCommand::Update()
+{
+	if (TestStepData->ScreenshotTask == nullptr)
+	{
+		const FString ScreenshotNameString = FString::Printf(TEXT("%s_%d"), *TestStepData->TestPlanData.AssetData.AssetName.ToString(), TestStepData->StepIndex);
+		TestStepData->ScreenshotTask = UAutomationBlueprintFunctionLibrary::TakeHighResScreenshot(1280, 720, ScreenshotNameString, nullptr, false, false, ScreenshotParameters.ComparisonTolerance);
+	}
+	
+	if (TestStepData->ScreenshotTask && !TestStepData->ScreenshotTask->IsValidTask())
+	{
+		if (FAutomationTestBase* CurrentTest = FAutomationTestFramework::Get().GetCurrentTest())
+		{
+			CurrentTest->AddError(TEXT("Screenshot Capture Task is not valid."));
+			TestStepData->bSuccess = false;
+		}
+		return true;
+	}
+
+	return TestStepData->ScreenshotTask && TestStepData->ScreenshotTask->IsTaskDone();
+}
+
+bool FInterchangeInterStepPerformTestsAndCollectGarbageCommand::Update()
+{
+	FAutomationTestBase* CurrentTest = FAutomationTestFramework::Get().GetCurrentTest();
+
+	// Restore Viewport Parameters
+	UAutomationBlueprintFunctionLibrary::SetEditorActiveViewportViewMode(TestStepData->CachedScreenshotParameters.ViewMode);
+	UAutomationBlueprintFunctionLibrary::SetEditorActiveViewportWireframeOpacity(TestStepData->CachedScreenshotParameters.WireframeOpacity);
+
+	// Reset the Screenshot task so that next step could start a new task for screenshot.
+	TestStepData->ScreenshotTask = nullptr;
+
+	TestStepData->TestPlanData.TestPlan->Steps[TestStepData->StepIndex]->FinishStep(TestStepData->TestPlanData, CurrentTest);
+	if (TestStepData->TestPlanData.InterchangeResults)
+	{
+		// Populate the automation test execution info with the interchange import results
+		for (UInterchangeResult* Result : TestStepData->TestPlanData.InterchangeResults->GetResults())
+		{
+			switch (Result->GetResultType())
+			{
+			case EInterchangeResultType::Error:
+			{
+				if (CurrentTest)
+				{
+					CurrentTest->AddError(Result->GetText().ToString());
+				}
+				TestStepData->bSuccess = false;
+				break;
+			}
+			case EInterchangeResultType::Warning:
+			{
+				if (CurrentTest)
+				{
+					CurrentTest->AddWarning(Result->GetText().ToString());
+				}
+				break;
+			}
+			}
+		}
+		TestStepData->TestPlanData.InterchangeResults->ClearFlags(RF_Standalone);
+	}
+
+	// Finished with the interchange results - null it so it will be GC'd later
+	TestStepData->TestPlanData.InterchangeResults = nullptr;
+
+	if (CurrentTest)
+	{
+		CurrentTest->PopContext();
+	}
+
+	// Collect garbage between every step, so that we remove renamed packages which come from the save+reload operation
+	// Note we also reset the transaction buffer here to stop it from holding onto references which would prevent garbage collection.
+	// @todo: not really a big fan of this; is there a better way of just disabling transactions?
+	if ((GEditor != nullptr) && (GEditor->Trans != nullptr))
+	{
+		GEditor->Trans->Reset(FText::FromString("Discard undo history during Automation testing."));
+	}
+	CollectGarbage(GARBAGE_COLLECTION_KEEPFLAGS);
+
+	return true;
+}
+
+// Called after all the steps.
+bool FInterchangeTestCleanUpCommand::Update()
+{
+	TArray<UObject*> ObjectsToDelete;
+	ObjectsToDelete.Reserve(ObjectsToDelete.Max() + TestStepData->TestPlanData.ResultObjects.Num());
+	for (UObject* ResultObject : TestStepData->TestPlanData.ResultObjects)
+	{
+		if (AActor* Actor = Cast<AActor>(ResultObject))
+		{
+			constexpr bool bShouldModifyLevel = true;
+			Actor->GetWorld()->EditorDestroyActor(Actor, bShouldModifyLevel);
+			// Call UObject::Rename directly on actor to avoid AActor::Rename which unnecessarily unregister and re-register components
+			Actor->UObject::Rename(nullptr, GetTransientPackage(), REN_DontCreateRedirectors);
+		}
+		else
+		{
+			if (ResultObject && !ResultObject->HasAnyFlags(RF_BeginDestroyed | RF_FinishDestroyed))
+			{
+				ObjectsToDelete.Add(ResultObject);
+			}
+		}
+	}
+	//Make sure all compilation is done before deleting some objects
+	FAssetCompilingManager::Get().FinishAllCompilation();
+
+	constexpr bool bShowConfirmation = false;
+	ObjectTools::ForceDeleteObjects(ObjectsToDelete, bShowConfirmation);
+	return true;
+
+}
+
+
+bool FInterchangeTestAutomationTestSuccessCommand::Update()
+{
+	if (FAutomationTestBase* CurrentTest = FAutomationTestFramework::Get().GetCurrentTest())
+	{
+		CurrentTest->TestTrue(TEXT("Interchange Import Automation Test Success"), TestStepData->bSuccess);
+	}
+	return true;
+}
\ No newline at end of file
diff --git a/Engine/Plugins/Tests/InterchangeTests/Source/InterchangeTests/Private/InterchangeAutomatedTestUtils.h b/Engine/Plugins/Tests/InterchangeTests/Source/InterchangeTests/Private/InterchangeAutomatedTestUtils.h
new file mode 100644
--- /dev/null
+++ b/Engine/Plugins/Tests/InterchangeTests/Source/InterchangeTests/Private/InterchangeAutomatedTestUtils.h
@@ -0,0 +1,36 @@
+// Copyright Epic Games, Inc. All Rights Reserved.
+#pragma once
+
+#include "Misc/AutomationTest.h"
+#include "InterchangeImportTestStepBase.h"
+#include "InterchangeManager.h"
+#include "InterchangeImportTestData.h"
+
+class UAutomationEditorTask;
+class UClass;
+
+struct FInterchangeAutomationTestStepData
+{
+	FInterchangeImportTestData TestPlanData;
+	FInterchangeTestScreenshotParameters CachedScreenshotParameters;
+	TTuple<UE::Interchange::FAssetImportResultPtr, UE::Interchange::FSceneImportResultPtr> Results;
+	UAutomationEditorTask* ScreenshotTask = nullptr;
+	int32 StepIndex = 0;
+	bool bSuccess = true;
+};
+
+using InterchangeAutomationTestStepDataRef = TSharedRef< FInterchangeAutomationTestStepData, ESPMode::ThreadSafe>;
+
+DEFINE_LATENT_AUTOMATION_COMMAND_TWO_PARAMETER(FInterchangeIntializeStepCommand, InterchangeAutomationTestStepDataRef, TestStepData, int32, StepIndex);
+
+DEFINE_LATENT_AUTOMATION_COMMAND_ONE_PARAMETER(FInterchangeInterStepCollectResultCommand, InterchangeAutomationTestStepDataRef, TestStepData);
+
+DEFINE_LATENT_AUTOMATION_COMMAND_ONE_PARAMETER(FInterchangeInterStepPerformTestsAndCollectGarbageCommand, InterchangeAutomationTestStepDataRef, TestStepData);
+
+DEFINE_LATENT_AUTOMATION_COMMAND_ONE_PARAMETER(FInterchangeTestCleanUpCommand, InterchangeAutomationTestStepDataRef, TestStepData);
+
+DEFINE_LATENT_AUTOMATION_COMMAND_THREE_PARAMETER(FInterchangeSetupScreenshotViewportCommand, InterchangeAutomationTestStepDataRef, TestStepData, TSet<UClass*>, ActorClassFilter, FInterchangeTestScreenshotParameters, ScreenshotParameters);
+
+DEFINE_LATENT_AUTOMATION_COMMAND_TWO_PARAMETER(FInterchangeCaptureScreenshotCommand, InterchangeAutomationTestStepDataRef, TestStepData, FInterchangeTestScreenshotParameters, ScreenshotParameters);
+
+DEFINE_LATENT_AUTOMATION_COMMAND_ONE_PARAMETER(FInterchangeTestAutomationTestSuccessCommand, InterchangeAutomationTestStepDataRef, TestStepData);
diff --git a/Engine/Plugins/Tests/InterchangeTests/Source/InterchangeTests/Private/InterchangeImportTest.cpp b/Engine/Plugins/Tests/InterchangeTests/Source/InterchangeTests/Private/InterchangeImportTest.cpp
--- a/Engine/Plugins/Tests/InterchangeTests/Source/InterchangeTests/Private/InterchangeImportTest.cpp
+++ b/Engine/Plugins/Tests/InterchangeTests/Source/InterchangeTests/Private/InterchangeImportTest.cpp
@@ -6,11 +6,13 @@
 #include "Editor.h"
 #include "Editor/Transactor.h"
 #include "HAL/FileManager.h"
-#include "Misc/AutomationTest.h"
+#include "InterchangeAutomatedTestUtils.h"
 #include "InterchangeHelper.h"
 #include "InterchangeImportTestData.h"
 #include "InterchangeImportTestPlan.h"
 #include "InterchangeImportTestStepBase.h"
+#include "Misc/AutomationTest.h"
+#include "Tests/AutomationCommon.h"
 #include "Modules/ModuleManager.h"
 #include "ObjectTools.h"
 
@@ -52,11 +54,16 @@ void FInterchangeImportTest::GetTests(TArray<FString>& OutBeautifiedNames, TArra
 
 bool FInterchangeImportTest::RunTest(const FString& Path)
 {
+	#define ADD_COMMAND(ClassDeclaration) if (bRunSynchornously) {ClassDeclaration.Update();}\
+	else{ADD_LATENT_AUTOMATION_COMMAND(ClassDeclaration)}
+
 	// For now, we can't run Interchange automation tests on Mac, as there's a problem with the InterchangeWorker. For now, do nothing on Mac.
 	// @todo: find a solution to this
 #if PLATFORM_MAC || PLATFORM_LINUX
 	return true;
 #endif
+	
+	TestParameterContext = FString::Printf(TEXT("Interchange.%s"), *UE::Interchange::FInterchangeImportTestPlanStaticHelpers::GetTestNameFromObjectPathString(Path));
 
 	static const auto CVarInterchangeFbx = IConsoleManager::Get().FindConsoleVariable(TEXT("Interchange.FeatureFlags.Import.FBX"));
 	bool IsInterchangeFbxEnabled = CVarInterchangeFbx->GetBool();
@@ -67,17 +74,18 @@ bool FInterchangeImportTest::RunTest(const FString& Path)
 	//Make sure interchange is enabled for fbx
 	CVarInterchangeFbx->Set(true, ECVF_SetByConsole);
 
+	TSharedPtr<FInterchangeAutomationTestStepData, ESPMode::ThreadSafe> AutomationTestStepData = MakeShared<FInterchangeAutomationTestStepData, ESPMode::ThreadSafe>();
 	
-	FInterchangeImportTestData TestPlanData;
 	FName PackageName = FName(FPaths::GetBaseFilename(Path, false));
 	FName PackagePath = FName(FPaths::GetPath(Path));
 	FName AssetName = FName(FPaths::GetBaseFilename(Path, true));
 	FTopLevelAssetPath ClassName = UInterchangeImportTestPlan::StaticClass()->GetClassPathName();
-	TestPlanData.AssetData = FAssetData(PackageName, PackagePath, AssetName, ClassName);
-	TestPlanData.TestPlan = CastChecked<UInterchangeImportTestPlan>(TestPlanData.AssetData.GetAsset());
+	AutomationTestStepData->TestPlanData.AssetData = FAssetData(PackageName, PackagePath, AssetName, ClassName);
+	AutomationTestStepData->TestPlanData.TestPlan = CastChecked<UInterchangeImportTestPlan>(AutomationTestStepData->TestPlanData.AssetData.GetAsset());
 	
-	if (!ensure(TestPlanData.TestPlan))
+	if (!ensure(AutomationTestStepData->TestPlanData.TestPlan))
 	{
+		AutomationTestStepData.Reset();
 		return false;
 	}
 
@@ -89,173 +97,52 @@ bool FInterchangeImportTest::RunTest(const FString& Path)
 
 	// Clear out the folder contents before we do anything else
 
-	const bool bRequireExists = false;
-	const bool bDeleteRecursively = true;
+	constexpr bool bRequireExists = false;
+	constexpr bool bDeleteRecursively = true;
 	IFileManager::Get().DeleteDirectory(*BaseFilePath, bRequireExists, bDeleteRecursively);
 
-	TTuple<UE::Interchange::FAssetImportResultPtr, UE::Interchange::FSceneImportResultPtr> Results;
+	FString TestPlanAssetName = AutomationTestStepData->TestPlanData.AssetData.AssetName.ToString().Replace(TEXT("/"), TEXT("_"));
+	AutomationTestStepData->TestPlanData.DestAssetPackagePath = BasePackagePath / TestPlanAssetName;
+	AutomationTestStepData->TestPlanData.DestAssetFilePath = BaseFilePath / TestPlanAssetName;
+	const bool bAddRecursively = true;
+	IFileManager::Get().MakeDirectory(*AutomationTestStepData->TestPlanData.DestAssetFilePath, bAddRecursively);
 
-	int32 StepIndex = 0;
-	bool bFoundValidSteps = true;
-	bool bSuccess = true;
+	AutomationTestStepData->StepIndex = 0;
+	AutomationTestStepData->bSuccess = true;
 
-	while (bFoundValidSteps)
+	const bool bRunSynchornously = AutomationTestStepData->TestPlanData.TestPlan->IsRunningSynchornously();
+	for (int32 TestStepIndex = 0; TestStepIndex < AutomationTestStepData->TestPlanData.TestPlan->Steps.Num(); ++TestStepIndex)
 	{
-		bFoundValidSteps = false;
-		TestPlanData.ImportedAssets.Empty();
-
-		FString TestPlanAssetName = TestPlanData.AssetData.AssetName.ToString().Replace(TEXT("/"), TEXT("_"));
-		TestPlanData.DestAssetPackagePath = BasePackagePath / TestPlanAssetName;
-		TestPlanData.DestAssetFilePath = BaseFilePath / TestPlanAssetName;
-
-		const bool bAddRecursively = true;
-		IFileManager::Get().MakeDirectory(*TestPlanData.DestAssetFilePath, bAddRecursively);
-
-		if (TestPlanData.TestPlan)
-		{
-			if (TestPlanData.TestPlan->Steps.IsValidIndex(StepIndex))
-			{
-				bFoundValidSteps = true;
-				Results = TestPlanData.TestPlan->Steps[StepIndex]->StartStep(TestPlanData);
-			}
-		}
-		
-
-		if (bFoundValidSteps)
+		if (UInterchangeImportTestStepBase* TestStepBase = AutomationTestStepData->TestPlanData.TestPlan->Steps[TestStepIndex])
 		{
-			UE::Interchange::FAssetImportResultPtr AssetImportResults = Results.Get<0>();
-			UE::Interchange::FSceneImportResultPtr SceneImportResults = Results.Get<1>();
+			ADD_COMMAND(FInterchangeIntializeStepCommand(AutomationTestStepData.ToSharedRef(), TestStepIndex));
+			ADD_COMMAND(FInterchangeInterStepCollectResultCommand(AutomationTestStepData.ToSharedRef()));
 
-			if (AssetImportResults)
+			// Screenshot Tests should only be performed via Test Automation Window
+			if (!bRunSynchornously && TestStepBase->HasScreenshotTest())
 			{
-				AssetImportResults->WaitUntilDone();
-			}
-
-			if (SceneImportResults)
-			{
-				SceneImportResults->WaitUntilDone();
-			}
-
-
-			// Start post-steps
-			if (TestPlanData.TestPlan)
-			{
-				if (TestPlanData.TestPlan->Steps.IsValidIndex(StepIndex))
-				{
-					FString Context = TestPlanData.AssetData.AssetName.ToString() + TEXT(": ") + TestPlanData.TestPlan->Steps[StepIndex]->GetContextString();
-					ExecutionInfo.PushContext(Context);
-
-					// Fill out list of result objects in the data object.
-					// These are the UObject* results corresponding to imported assets.
-					if (AssetImportResults)
-					{
-						for (UObject* ImportedObject : AssetImportResults->GetImportedObjects())
-						{
-							TestPlanData.ResultObjects.AddUnique(ImportedObject);
-						}
-
-						// Also add the InterchangeResultsContainer to the data so that tests can be run on it
-						// (e.g. to check whether something imported with a specific expected error)
-						TestPlanData.InterchangeResults = AssetImportResults->GetResults();
-
-						// Fill out list of imported assets as FAssetData
-						for (UObject* Object : AssetImportResults->GetImportedObjects())
-						{
-							TestPlanData.ImportedAssets.Emplace(Object);
-						}
-					}
-
-					if (SceneImportResults)
-					{
-						for (UObject* ImportedObject : SceneImportResults->GetImportedObjects())
-						{
-							TestPlanData.ResultObjects.AddUnique(ImportedObject);
-						}
-
-						if (TestPlanData.InterchangeResults)
-						{
-							TestPlanData.InterchangeResults->Append(SceneImportResults->GetResults());
-						}
-						else
-						{
-							TestPlanData.InterchangeResults = SceneImportResults->GetResults();
-						}
-					}
-
-					// We don't need the Interchange results object any more, as we've already taken everything from it that we need.
-					// If we don't reset it, it will hold on to the trashed versions of the imported assets during GC.
-					Results = {};
-					//Make sure interchange result are not garbage collected in case we reload some packages
-					if (TestPlanData.InterchangeResults)
-					{
-						TestPlanData.InterchangeResults->SetFlags(RF_Standalone);
-					}
-					TestPlanData.TestPlan->Steps[StepIndex]->FinishStep(TestPlanData, ExecutionInfo);
-
-					if (TestPlanData.InterchangeResults)
-					{
-						// Populate the automation test execution info with the interchange import results
-						for (UInterchangeResult* Result : TestPlanData.InterchangeResults->GetResults())
-						{
-							switch (Result->GetResultType())
-							{
-							case EInterchangeResultType::Error:
-								ExecutionInfo.AddError(Result->GetText().ToString());
-								bSuccess = false;
-								break;
-
-							case EInterchangeResultType::Warning:
-								ExecutionInfo.AddWarning(Result->GetText().ToString());
-								break;
-							}
-						}
-						TestPlanData.InterchangeResults->ClearFlags(RF_Standalone);
-					}
-
-					// Finished with the interchange results - null it so it will be GC'd later
-					TestPlanData.InterchangeResults = nullptr;
-
-					ExecutionInfo.PopContext();
-				}
-			}
-
-			// Collect garbage between every step, so that we remove renamed packages which come from the save+reload operation
-			// Note we also reset the transaction buffer here to stop it from holding onto references which would prevent garbage collection.
-			// @todo: not really a big fan of this; is there a better way of just disabling transactions?
-			if ((GEditor != nullptr) && (GEditor->Trans != nullptr))
-			{
-				GEditor->Trans->Reset(FText::FromString("Discard undo history during FBX Automation testing."));
-			}
-			CollectGarbage(GARBAGE_COLLECTION_KEEPFLAGS);
-		}
-
-		StepIndex++;
-	}
-
-	TArray<UObject*> ObjectsToDelete;
-	ObjectsToDelete.Reserve(ObjectsToDelete.Max() + TestPlanData.ResultObjects.Num());
-	for (UObject* ResultObject : TestPlanData.ResultObjects)
-	{
-		if (AActor* Actor = Cast<AActor>(ResultObject))
-		{
-			constexpr bool bShouldModifyLevel = true;
-			Actor->GetWorld()->EditorDestroyActor(Actor, bShouldModifyLevel);
-			// Call UObject::Rename directly on actor to avoid AActor::Rename which unnecessarily unregister and re-register components
-			Actor->UObject::Rename(nullptr, GetTransientPackage(), REN_DontCreateRedirectors);
-		}
-		else
-		{
-			if (ResultObject && !ResultObject->HasAnyFlags(RF_BeginDestroyed | RF_FinishDestroyed))
-			{
-				ObjectsToDelete.Add(ResultObject);
+				// Wait so that all the materials applied to actors would be visible in the screenshot.
+				ADD_LATENT_AUTOMATION_COMMAND(FWaitForShadersToFinishCompilingInGame());
+				
+				const FInterchangeTestScreenshotParameters ScreenshotParameters = TestStepBase->GetScreenshotParameters();
+
+				// Focuses camera on relevant actors in the scene.
+				ADD_LATENT_AUTOMATION_COMMAND(FInterchangeSetupScreenshotViewportCommand(AutomationTestStepData.ToSharedRef(), TestStepBase->GetFocusActorsAllowedClasses(), ScreenshotParameters));
+				
+				// Camera Viewport Transition Time.
+				ADD_LATENT_AUTOMATION_COMMAND(FEngineWaitLatentCommand(0.3f));
+
+				// Send a request for screenshot and wait until comparison is finished.
+				ADD_LATENT_AUTOMATION_COMMAND(FInterchangeCaptureScreenshotCommand(AutomationTestStepData.ToSharedRef(), ScreenshotParameters));
 			}
+			
+			ADD_COMMAND(FInterchangeInterStepPerformTestsAndCollectGarbageCommand(AutomationTestStepData.ToSharedRef()));
 		}
 	}
-	//Make sure all compilation is done before deleting some objects
-	FAssetCompilingManager::Get().FinishAllCompilation();
-
-	constexpr bool bShowConfirmation = false;
-	ObjectTools::ForceDeleteObjects(ObjectsToDelete, bShowConfirmation);
+	ADD_COMMAND(FInterchangeTestCleanUpCommand(AutomationTestStepData.ToSharedRef()));
+	ADD_COMMAND(FInterchangeTestAutomationTestSuccessCommand(AutomationTestStepData.ToSharedRef()));
+	
+	return true;
 
-	return bSuccess;
+#undef ADD_COMMAND
 }
diff --git a/Engine/Plugins/Tests/InterchangeTests/Source/InterchangeTests/Private/InterchangeImportTestPlan.cpp b/Engine/Plugins/Tests/InterchangeTests/Source/InterchangeTests/Private/InterchangeImportTestPlan.cpp
--- a/Engine/Plugins/Tests/InterchangeTests/Source/InterchangeTests/Private/InterchangeImportTestPlan.cpp
+++ b/Engine/Plugins/Tests/InterchangeTests/Source/InterchangeTests/Private/InterchangeImportTestPlan.cpp
@@ -50,6 +50,7 @@ void UInterchangeImportTestPlan::WriteToJson(const FString& Filename)
 
 void UInterchangeImportTestPlan::RunThisTest()
 {
+	bRunSynchronously = true;
 	FMessageLog AutomationEditorLog("AutomationTestingLog");
 	FString NewPageName = FString::Printf(TEXT("----- Interchange Import Test: %s----"), *GetPathName());
 	FText NewPageNameText = FText::FromString(*NewPageName);
@@ -86,6 +87,13 @@ void UInterchangeImportTestPlan::RunThisTest()
 			}
 		}
 	}
+
+	bRunSynchronously = false;
+}
+
+bool UInterchangeImportTestPlan::IsRunningSynchornously() const
+{
+	return bRunSynchronously;
 }
 
 namespace UE::Interchange
diff --git a/Engine/Plugins/Tests/InterchangeTests/Source/InterchangeTests/Private/InterchangeImportTestStepBase.cpp b/Engine/Plugins/Tests/InterchangeTests/Source/InterchangeTests/Private/InterchangeImportTestStepBase.cpp
--- a/Engine/Plugins/Tests/InterchangeTests/Source/InterchangeTests/Private/InterchangeImportTestStepBase.cpp
+++ b/Engine/Plugins/Tests/InterchangeTests/Source/InterchangeTests/Private/InterchangeImportTestStepBase.cpp
@@ -13,7 +13,7 @@
 #include UE_INLINE_GENERATED_CPP_BY_NAME(InterchangeImportTestStepBase)
 
 
-bool UInterchangeImportTestStepBase::PerformTests(FInterchangeImportTestData& Data, FAutomationTestExecutionInfo& ExecutionInfo)
+bool UInterchangeImportTestStepBase::PerformTests(FInterchangeImportTestData& Data, FAutomationTestBase* CurrentTest)
 {
 	// Add the Interchange results container to the list of UObject* we pass to the test functions.
 	// This way we can match against tests which can operate on the UInterchangeResultsContainer too.
@@ -25,15 +25,21 @@ bool UInterchangeImportTestStepBase::PerformTests(FInterchangeImportTestData& Da
 	for (FInterchangeTestFunction& Test : Tests)
 	{
 		FInterchangeTestFunctionResult Result = Test.Invoke(ResultObjects);
-
+		
 		for (const FString& Warning : Result.GetWarnings())
 		{
-			ExecutionInfo.AddWarning(Warning);
+			if (CurrentTest)
+			{
+				CurrentTest->AddWarning(Warning);
+			}
 		}
 
 		for (const FString& Error : Result.GetErrors())
 		{
-			ExecutionInfo.AddError(Error);
+			if (CurrentTest)
+			{
+				CurrentTest->AddError(Error);
+			}
 		}
 
 		bSuccess &= Result.IsSuccess();
diff --git a/Engine/Plugins/Tests/InterchangeTests/Source/InterchangeTests/Private/InterchangeImportTestStepImport.cpp b/Engine/Plugins/Tests/InterchangeTests/Source/InterchangeTests/Private/InterchangeImportTestStepImport.cpp
--- a/Engine/Plugins/Tests/InterchangeTests/Source/InterchangeTests/Private/InterchangeImportTestStepImport.cpp
+++ b/Engine/Plugins/Tests/InterchangeTests/Source/InterchangeTests/Private/InterchangeImportTestStepImport.cpp
@@ -13,6 +13,7 @@
 #include "InterchangeManager.h"
 #include "InterchangeProjectSettings.h"
 #include "InterchangeTestsLog.h"
+#include "LevelEditorSubsystem.h"
 #include "Misc/PackageName.h"
 #include "Misc/ScopedSlowTask.h"
 #include "ScopedTransaction.h"
@@ -95,23 +96,34 @@ TTuple<UE::Interchange::FAssetImportResultPtr, UE::Interchange::FSceneImportResu
 	if (bImportIntoLevel)
 	{
 		// Create transient world to host data from producer
-		TransientWorld = TStrongObjectPtr<UWorld>(NewObject<UWorld>(GetTransientPackage()));
-		TransientWorld->WorldType = EWorldType::EditorPreview;
-
-		FWorldContext& WorldContext = GEngine->CreateNewWorldContext(TransientWorld->WorldType);
-		WorldContext.SetCurrentWorld(TransientWorld.Get());
-
-		TransientWorld->InitializeNewWorld(UWorld::InitializationValues()
-			.AllowAudioPlayback(false)
-			.CreatePhysicsScene(false)
-			.RequiresHitProxies(false)
-			.CreateNavigation(false)
-			.CreateAISystem(false)
-			.ShouldSimulatePhysics(false)
-			.SetTransactional(false)
-		);
-
-		Params.ImportLevel = TransientWorld->GetCurrentLevel();
+		if (WorldPath.IsValid())
+		{
+			if (ULevelEditorSubsystem* LevelEditorSubsystem = GEditor->GetEditorSubsystem<ULevelEditorSubsystem>())
+			{
+				LevelEditorSubsystem->LoadLevel(WorldPath.GetAssetPathString());
+				Params.ImportLevel = GWorld->GetCurrentLevel();
+			}
+		}
+		else 
+		{
+			TransientWorld = TStrongObjectPtr<UWorld>(NewObject<UWorld>(GetTransientPackage()));
+			TransientWorld->WorldType = EWorldType::EditorPreview;
+
+			FWorldContext& WorldContext = GEngine->CreateNewWorldContext(TransientWorld->WorldType);
+			WorldContext.SetCurrentWorld(TransientWorld.Get());
+
+			TransientWorld->InitializeNewWorld(UWorld::InitializationValues()
+				.AllowAudioPlayback(false)
+				.CreatePhysicsScene(false)
+				.RequiresHitProxies(false)
+				.CreateNavigation(false)
+				.CreateAISystem(false)
+				.ShouldSimulatePhysics(false)
+				.SetTransactional(false)
+			);
+
+			Params.ImportLevel = TransientWorld->GetCurrentLevel();
+		}
 
 		return InterchangeManager.ImportSceneAsync(Data.DestAssetPackagePath, ScopedSourceData.GetSourceData(), Params);
 	}
@@ -122,8 +134,7 @@ TTuple<UE::Interchange::FAssetImportResultPtr, UE::Interchange::FSceneImportResu
 	}
 }
 
-
-FTestStepResults UInterchangeImportTestStepImport::FinishStep(FInterchangeImportTestData& Data, FAutomationTestExecutionInfo& ExecutionInfo)
+FTestStepResults UInterchangeImportTestStepImport::FinishStep(FInterchangeImportTestData& Data, FAutomationTestBase* CurrentTest)
 {
 	FTestStepResults Results;
 
@@ -135,38 +146,67 @@ FTestStepResults UInterchangeImportTestStepImport::FinishStep(FInterchangeImport
 	}
 
 	// Run all the tests
-	bool bSuccess = PerformTests(Data, ExecutionInfo);
-
-if (TransientWorld)
-{
+	bool bSuccess = PerformTests(Data, CurrentTest);
+	
 	// Only keep assets as result objects since the world and its actors are being destroyed
 	TArray<UObject*> ResultObjects = MoveTemp(Data.ResultObjects);
 	Data.ResultObjects.Reset(ResultObjects.Num());
 
 	for (UObject* Object : ResultObjects)
 	{
-		if (!Object->IsA<AActor>())
+		if (!Object->IsA<AActor>() && !Object->IsA<UWorld>())
 		{
 			Data.ResultObjects.Add(Object);
 		}
 	}
 
-	// Now delete world
-	GEngine->DestroyWorldContext(TransientWorld.Get());
-	TransientWorld->DestroyWorld(true);
-	TransientWorld.Reset();
-}
+	if (WorldPath.IsValid())
+	{
+		if (ULevelEditorSubsystem* LevelEditorSubsystem = GEditor->GetEditorSubsystem<ULevelEditorSubsystem>())
+		{
+			LevelEditorSubsystem->LoadLevel(WorldPath.GetAssetPathString());
+		}
+	}
+	else if (TransientWorld)
+	{
+		// Now delete world
+		GEngine->DestroyWorldContext(TransientWorld.Get());
+		TransientWorld->DestroyWorld(true);
+		TransientWorld.Reset();
+	}
 
-Results.bTestStepSuccess = bSuccess;
-return Results;
+	Results.bTestStepSuccess = bSuccess;
+	return Results;
 }
 
-
 FString UInterchangeImportTestStepImport::GetContextString() const
 {
 	return FString(TEXT("Importing ")) + FPaths::GetCleanFilename(SourceFile.FilePath);
 }
 
+bool UInterchangeImportTestStepImport::HasScreenshotTest() const
+{
+	return bImportIntoLevel && WorldPath.IsValid();
+}
+
+FInterchangeTestScreenshotParameters UInterchangeImportTestStepImport::GetScreenshotParameters() const
+{
+	return ScreenshotParameters;
+}
+
+TSet<UClass*> UInterchangeImportTestStepImport::GetFocusActorsAllowedClasses() const
+{
+	TSet<UClass*> ActorClasses;
+	for (const TSubclassOf<AActor>& ActorSubclass : FocusActorClasses)
+	{
+		if (UClass* ActorClass = ActorSubclass.Get())
+		{
+			ActorClasses.Add(ActorClass);
+		}
+	}
+	return ActorClasses;
+}
+
 bool UInterchangeImportTestStepImport::CanEditPipelineSettings() const
 {
 	if (SourceFile.FilePath.IsEmpty())
diff --git a/Engine/Plugins/Tests/InterchangeTests/Source/InterchangeTests/Private/InterchangeImportTestStepReimport.cpp b/Engine/Plugins/Tests/InterchangeTests/Source/InterchangeTests/Private/InterchangeImportTestStepReimport.cpp
--- a/Engine/Plugins/Tests/InterchangeTests/Source/InterchangeTests/Private/InterchangeImportTestStepReimport.cpp
+++ b/Engine/Plugins/Tests/InterchangeTests/Source/InterchangeTests/Private/InterchangeImportTestStepReimport.cpp
@@ -64,7 +64,7 @@ TTuple<UE::Interchange::FAssetImportResultPtr, UE::Interchange::FSceneImportResu
 }
 
 
-FTestStepResults UInterchangeImportTestStepReimport::FinishStep(FInterchangeImportTestData& Data, FAutomationTestExecutionInfo& ExecutionInfo)
+FTestStepResults UInterchangeImportTestStepReimport::FinishStep(FInterchangeImportTestData& Data, FAutomationTestBase* CurrentTest)
 {
 	FTestStepResults Results;
 
@@ -75,7 +75,7 @@ FTestStepResults UInterchangeImportTestStepReimport::FinishStep(FInterchangeImpo
 	}
 
 	// Run all the tests
-	bool bSuccess = PerformTests(Data, ExecutionInfo);
+	bool bSuccess = PerformTests(Data, CurrentTest);
 
 	Results.bTestStepSuccess = bSuccess;
 	return Results;
@@ -87,6 +87,20 @@ FString UInterchangeImportTestStepReimport::GetContextString() const
 	return FString(TEXT("Reimporting ")) + FPaths::GetCleanFilename(SourceFileToReimport.FilePath);
 }
 
+bool UInterchangeImportTestStepReimport::HasScreenshotTest() const
+{
+	return false;
+}
+
+FInterchangeTestScreenshotParameters UInterchangeImportTestStepReimport::GetScreenshotParameters() const
+{
+	return FInterchangeTestScreenshotParameters();
+}
+
+TSet<UClass*> UInterchangeImportTestStepReimport::GetFocusActorsAllowedClasses() const
+{
+	return TSet<UClass*>();
+}
 
 bool  UInterchangeImportTestStepReimport::CanEditPipelineSettings() const
 {
diff --git a/Engine/Plugins/Tests/InterchangeTests/Source/InterchangeTests/Public/InterchangeImportTestPlan.h b/Engine/Plugins/Tests/InterchangeTests/Source/InterchangeTests/Public/InterchangeImportTestPlan.h
--- a/Engine/Plugins/Tests/InterchangeTests/Source/InterchangeTests/Public/InterchangeImportTestPlan.h
+++ b/Engine/Plugins/Tests/InterchangeTests/Source/InterchangeTests/Public/InterchangeImportTestPlan.h
@@ -31,6 +31,11 @@ public:
 	/** Click here to immediately run this single test through the automation framework */
 	UFUNCTION(BlueprintCallable, CallInEditor, Category = Run, Meta = (DisplayPriority = 1))
 	void RunThisTest();
+
+	bool IsRunningSynchornously() const;
+
+private:
+	bool bRunSynchronously = false;
 };
 
 
diff --git a/Engine/Plugins/Tests/InterchangeTests/Source/InterchangeTests/Public/InterchangeImportTestStepBase.h b/Engine/Plugins/Tests/InterchangeTests/Source/InterchangeTests/Public/InterchangeImportTestStepBase.h
--- a/Engine/Plugins/Tests/InterchangeTests/Source/InterchangeTests/Public/InterchangeImportTestStepBase.h
+++ b/Engine/Plugins/Tests/InterchangeTests/Source/InterchangeTests/Public/InterchangeImportTestStepBase.h
@@ -3,6 +3,9 @@
 #pragma once
 
 #include "InterchangeManager.h"
+#include "AutomationScreenshotOptions.h"
+#include "Containers/EnumAsByte.h"
+#include "Engine/EngineBaseTypes.h"
 #include "InterchangeImportTestStepBase.generated.h"
 
 class FAutomationTestExecutionInfo;
@@ -16,6 +19,22 @@ struct FTestStepResults
 	bool bTriggerGC = false;
 };
 
+USTRUCT(BlueprintType)
+struct FInterchangeTestScreenshotParameters
+{
+	GENERATED_BODY()
+
+public:
+	UPROPERTY(EditAnywhere, Category = General)
+	EComparisonTolerance ComparisonTolerance = EComparisonTolerance::Low;
+
+	UPROPERTY(EditAnywhere, Category = General)
+	TEnumAsByte<EViewModeIndex> ViewMode = EViewModeIndex::VMI_Lit;
+
+	UPROPERTY(EditAnywhere, Category = General)
+	float WireframeOpacity = 0.2f;
+};
+
 
 UCLASS(BlueprintType, EditInlineNew, Abstract, autoExpandCategories = (General, Test))
 class INTERCHANGETESTS_API UInterchangeImportTestStepBase : public UObject
@@ -30,14 +49,18 @@ public:
 public:
 	virtual TTuple<UE::Interchange::FAssetImportResultPtr, UE::Interchange::FSceneImportResultPtr>
 		StartStep(FInterchangeImportTestData& Data) PURE_VIRTUAL(UInterchangeImportTestStepBase::StartStep, return {}; );
-	virtual FTestStepResults FinishStep(FInterchangeImportTestData& Data, FAutomationTestExecutionInfo& ExecutionInfo) PURE_VIRTUAL(UInterchangeImportTestStepBase::FinishStep, return {}; );
+	virtual FTestStepResults FinishStep(FInterchangeImportTestData& Data, FAutomationTestBase* CurrentTest) PURE_VIRTUAL(UInterchangeImportTestStepBase::FinishStep, return {}; );
 	virtual FString GetContextString() const PURE_VIRTUAL(UInterchangeImportTestStepBase::GetContextString, return {}; );
+	virtual bool HasScreenshotTest() const PURE_VIRTUAL(UInterchangeImportTestStepBase::HasScreenshotTest, return false; );
+	virtual FInterchangeTestScreenshotParameters GetScreenshotParameters() const PURE_VIRTUAL(UInterchangeImportTestStepBase::GetScreenshotParameters, return FInterchangeTestScreenshotParameters{}; );
 
 	virtual bool CanEditPipelineSettings() const PURE_VIRTUAL(UInterchangeImportTestStepBase::CanEditPipelineSettings, return false; );
 	virtual void EditPipelineSettings() PURE_VIRTUAL(UInterchangeImportTestStepBase::EditPipelineSettings);
 	virtual void ClearPipelineSettings() PURE_VIRTUAL(UInterchangeImportTestStepBase::ClearPipelineSettings);
 
-	bool PerformTests(FInterchangeImportTestData& Data, FAutomationTestExecutionInfo& ExecutionInfo);
+	virtual TSet<UClass*> GetFocusActorsAllowedClasses() const PURE_VIRTUAL(UInterchangeImportTestStepBase::GetFocusActorsAllowedClasses, return TSet<UClass*>(); );
+
+	bool PerformTests(FInterchangeImportTestData& Data, FAutomationTestBase* CurrentTest);
 protected:
 	void SaveReloadAssets(FInterchangeImportTestData& Data);
 };
diff --git a/Engine/Plugins/Tests/InterchangeTests/Source/InterchangeTests/Public/InterchangeImportTestStepImport.h b/Engine/Plugins/Tests/InterchangeTests/Source/InterchangeTests/Public/InterchangeImportTestStepImport.h
--- a/Engine/Plugins/Tests/InterchangeTests/Source/InterchangeTests/Public/InterchangeImportTestStepImport.h
+++ b/Engine/Plugins/Tests/InterchangeTests/Source/InterchangeTests/Public/InterchangeImportTestStepImport.h
@@ -11,7 +11,7 @@
 class UInterchangePipelineBase;
 class UWorld;
 
-UCLASS(BlueprintType, Meta = (DisplayName = "Import a file"))
+UCLASS(BlueprintType, AutoExpandCategories = ("General", "Screenshot Comparison"), Meta = (DisplayName = "Import a file"))
 class INTERCHANGETESTS_API UInterchangeImportTestStepImport : public UInterchangeImportTestStepBase
 {
 	GENERATED_BODY()
@@ -46,12 +46,25 @@ public:
 	UPROPERTY(EditAnywhere, Category = General)
 	bool bImportIntoLevel = false;
 
+	/**  Screen Shot Comparison when Importing in Level */
+	UPROPERTY(EditAnywhere, Category = "Screenshot Comparison", meta = (AllowedClasses = "/Script/Engine.World", EditCondition = "bImportIntoLevel"))
+	FSoftObjectPath WorldPath;
+
+	UPROPERTY(EditAnywhere, Category = "Screenshot Comparison", meta = (EditCondition = "bImportIntoLevel"))
+	TArray<TSubclassOf<AActor>> FocusActorClasses;
+
+	UPROPERTY(EditAnywhere, Category = "Screenshot Comparison", meta = (EditCondition = "bImportIntoLevel"))
+	FInterchangeTestScreenshotParameters ScreenshotParameters;
+
 public:
 	// UInterchangeImportTestStepBase interface
 	virtual TTuple<UE::Interchange::FAssetImportResultPtr, UE::Interchange::FSceneImportResultPtr>
 		StartStep(FInterchangeImportTestData& Data) override;
-	virtual FTestStepResults FinishStep(FInterchangeImportTestData& Data, FAutomationTestExecutionInfo& ExecutionInfo) override;
+	virtual FTestStepResults FinishStep(FInterchangeImportTestData& Data, FAutomationTestBase* CurrentTest) override;
 	virtual FString GetContextString() const override;
+	virtual bool HasScreenshotTest() const override;
+	virtual FInterchangeTestScreenshotParameters GetScreenshotParameters() const override;
+	virtual TSet<UClass*> GetFocusActorsAllowedClasses() const override;
 
 	virtual bool CanEditPipelineSettings() const override;
 	virtual void EditPipelineSettings() override;
diff --git a/Engine/Plugins/Tests/InterchangeTests/Source/InterchangeTests/Public/InterchangeImportTestStepReimport.h b/Engine/Plugins/Tests/InterchangeTests/Source/InterchangeTests/Public/InterchangeImportTestStepReimport.h
--- a/Engine/Plugins/Tests/InterchangeTests/Source/InterchangeTests/Public/InterchangeImportTestStepReimport.h
+++ b/Engine/Plugins/Tests/InterchangeTests/Source/InterchangeTests/Public/InterchangeImportTestStepReimport.h
@@ -43,8 +43,11 @@ public:
 	// UInterchangeImportTestStepBase interface
 	virtual TTuple<UE::Interchange::FAssetImportResultPtr, UE::Interchange::FSceneImportResultPtr>
 		StartStep(FInterchangeImportTestData& Data) override;
-	virtual FTestStepResults FinishStep(FInterchangeImportTestData& Data, FAutomationTestExecutionInfo& ExecutionInfo) override;
+	virtual FTestStepResults FinishStep(FInterchangeImportTestData& Data, FAutomationTestBase* CurrentTest) override;
 	virtual FString GetContextString() const override;
+	virtual bool HasScreenshotTest() const override;
+	virtual FInterchangeTestScreenshotParameters GetScreenshotParameters() const override;
+	virtual TSet<UClass*> GetFocusActorsAllowedClasses() const override;
 
 	virtual bool CanEditPipelineSettings() const override;
 	virtual void EditPipelineSettings() override;
diff --git a/Engine/Source/Developer/FunctionalTesting/Private/AutomationBlueprintFunctionLibrary.cpp b/Engine/Source/Developer/FunctionalTesting/Private/AutomationBlueprintFunctionLibrary.cpp
--- a/Engine/Source/Developer/FunctionalTesting/Private/AutomationBlueprintFunctionLibrary.cpp
+++ b/Engine/Source/Developer/FunctionalTesting/Private/AutomationBlueprintFunctionLibrary.cpp
@@ -1481,6 +1481,64 @@ void UAutomationBlueprintFunctionLibrary::SetEditorViewportViewMode(EViewModeInd
 #endif
 }
 
+void UAutomationBlueprintFunctionLibrary::SetEditorActiveViewportViewMode(EViewModeIndex Index)
+{
+#if WITH_EDITOR
+	FLevelEditorModule& LevelEditorModule = FModuleManager::GetModuleChecked<FLevelEditorModule>("LevelEditor");
+	if (SLevelViewport* LevelViewport = LevelEditorModule.GetFirstActiveLevelViewport().Get())
+	{
+		if (TSharedPtr<FEditorViewportClient> Viewport = LevelViewport->GetViewportClient())
+		{
+			Viewport->SetViewMode(Index);
+		}
+	}
+#endif
+}
+
+EViewModeIndex UAutomationBlueprintFunctionLibrary::GetEditorActiveViewportViewMode()
+{
+#if WITH_EDITOR
+	FLevelEditorModule& LevelEditorModule = FModuleManager::GetModuleChecked<FLevelEditorModule>("LevelEditor");
+	if (SLevelViewport* LevelViewport = LevelEditorModule.GetFirstActiveLevelViewport().Get())
+	{
+		if (TSharedPtr<FEditorViewportClient> Viewport = LevelViewport->GetViewportClient())
+		{
+			return Viewport->GetViewMode();
+		}
+	}
+#endif
+	return EViewModeIndex::VMI_Unknown;
+}
+
+void UAutomationBlueprintFunctionLibrary::SetEditorActiveViewportWireframeOpacity(float Opacity)
+{
+#if WITH_EDITOR
+	FLevelEditorModule& LevelEditorModule = FModuleManager::GetModuleChecked<FLevelEditorModule>("LevelEditor");
+	if (SLevelViewport* LevelViewport = LevelEditorModule.GetFirstActiveLevelViewport().Get())
+	{
+		if (TSharedPtr<FEditorViewportClient> Viewport = LevelViewport->GetViewportClient())
+		{
+			Viewport->WireframeOpacity = Opacity;
+		}
+	}
+#endif
+}
+
+float UAutomationBlueprintFunctionLibrary::GetEditorActiveViewportWireframeOpacity()
+{
+#if WITH_EDITOR
+	FLevelEditorModule& LevelEditorModule = FModuleManager::GetModuleChecked<FLevelEditorModule>("LevelEditor");
+	if (SLevelViewport* LevelViewport = LevelEditorModule.GetFirstActiveLevelViewport().Get())
+	{
+		if (TSharedPtr<FEditorViewportClient> Viewport = LevelViewport->GetViewportClient())
+		{
+			return Viewport->WireframeOpacity;
+		}
+	}
+#endif
+	return 0.0f;
+}
+
 void UAutomationBlueprintFunctionLibrary::SetEditorViewportVisualizeBuffer( FName BufferName )
 {
 #if WITH_EDITOR
diff --git a/Engine/Source/Developer/FunctionalTesting/Public/AutomationBlueprintFunctionLibrary.h b/Engine/Source/Developer/FunctionalTesting/Public/AutomationBlueprintFunctionLibrary.h
--- a/Engine/Source/Developer/FunctionalTesting/Public/AutomationBlueprintFunctionLibrary.h
+++ b/Engine/Source/Developer/FunctionalTesting/Public/AutomationBlueprintFunctionLibrary.h
@@ -215,6 +215,22 @@ public:
 	UFUNCTION(BlueprintCallable, Category = "Automation")
 	static void SetEditorViewportViewMode(EViewModeIndex Index);
 
+	/** Sets ViewMode (Lit/Unlit/etc.) of the first active viewport found **/
+	UFUNCTION(BlueprintCallable, Category = "Automation")
+	static void SetEditorActiveViewportViewMode(EViewModeIndex Index);
+
+	/** Get the ViewMode (Lit/Unlit/etc.) of the active viewport. Returns VMI_Unknown if fails to get the active viewport view mode index. **/
+	UFUNCTION(BlueprintCallable, Category = "Automation")
+	static EViewModeIndex GetEditorActiveViewportViewMode();
+
+	/** Sets the opacity for wireframe viewport view modes of the active viewport. **/
+	UFUNCTION(BlueprintCallable, Category = "Automation")
+	static void SetEditorActiveViewportWireframeOpacity(float Opacity);
+
+	/** Retrieves the opacity for wireframe viewport view modes of the active viewport. **/
+	UFUNCTION(BlueprintCallable, Category = "Automation")
+	static float GetEditorActiveViewportWireframeOpacity();
+
 	/** Sets all viewports of the first found level editor to have the VisualizeBuffer ViewMode and also display a given buffer (BaseColor/Metallic/Roughness/etc.) **/
 	UFUNCTION(BlueprintCallable, Category = "Automation")
 	static void SetEditorViewportVisualizeBuffer(FName BufferName);
