Interchange Test plan - Static Mesh - "Check socket location" not working

#rnx

diff --git a/Engine/Plugins/Tests/InterchangeTests/Source/InterchangeTests/Private/ImportTestFunctions/StaticMeshImportTestFunctions.cpp b/Engine/Plugins/Tests/InterchangeTests/Source/InterchangeTests/Private/ImportTestFunctions/StaticMeshImportTestFunctions.cpp
--- a/Engine/Plugins/Tests/InterchangeTests/Source/InterchangeTests/Private/ImportTestFunctions/StaticMeshImportTestFunctions.cpp
+++ b/Engine/Plugins/Tests/InterchangeTests/Source/InterchangeTests/Private/ImportTestFunctions/StaticMeshImportTestFunctions.cpp
@@ -743,24 +743,24 @@ FInterchangeTestFunctionResult UStaticMeshImportTestFunctions::CheckSocketName(U
 	return Result;
 }
 
-
 FInterchangeTestFunctionResult UStaticMeshImportTestFunctions::CheckSocketLocation(UStaticMesh* Mesh, int32 SocketIndex, const FVector& ExpectedSocketLocation)
 {
 	FInterchangeTestFunctionResult Result;
 
-	int32 SocketCount = Mesh->Sockets.Num();
-	if (SocketIndex >= SocketCount)
+	if (Mesh->Sockets.IsValidIndex(SocketIndex))
 	{
-		Result.AddError(FString::Printf(TEXT("The imported mesh doesn't contain %d sockets."), SocketIndex));
-	}
+		using namespace InterchangeTestsMathUtilities;
 
-	using namespace InterchangeTestsMathUtilities;
-
-	const FVector SocketLocation = RoundVectorToDecimalPlaces(Mesh->Sockets[SocketIndex]->RelativeLocation);
-	const FVector ExpectedSocketLocationRounded = RoundVectorToDecimalPlaces(ExpectedSocketLocation);
-	if (SocketLocation.Equals(ExpectedSocketLocationRounded, UE_DOUBLE_KINDA_SMALL_NUMBER))
+		const FVector SocketLocation = RoundVectorToDecimalPlaces(Mesh->Sockets[SocketIndex]->RelativeLocation);
+		const FVector ExpectedSocketLocationRounded = RoundVectorToDecimalPlaces(ExpectedSocketLocation);
+		if (!SocketLocation.Equals(ExpectedSocketLocationRounded, UE_DOUBLE_KINDA_SMALL_NUMBER))
+		{
+			Result.AddError(FString::Printf(TEXT("Expected socket location %s, imported %s"), *ExpectedSocketLocation.ToString(), *SocketLocation.ToString()));
+		}
+	}
+	else 
 	{
-		Result.AddError(FString::Printf(TEXT("Expected socket location %s, imported %s"), *ExpectedSocketLocation.ToString(), *SocketLocation.ToString()));
+		Result.AddError(FString::Printf(TEXT("The imported mesh doesn't contain socket with index %d."), SocketIndex));
 	}
 
 	return Result;
