[Interchange Test Plan] Skeletal Mesh SocketTest Functions
- They mirror the socket test functions from the StaticMesh

#rnx

diff --git a/Engine/Plugins/Tests/InterchangeTests/Source/InterchangeTests/Private/ImportTestFunctions/SkeletalMeshImportTestFunctions.cpp b/Engine/Plugins/Tests/InterchangeTests/Source/InterchangeTests/Private/ImportTestFunctions/SkeletalMeshImportTestFunctions.cpp
--- a/Engine/Plugins/Tests/InterchangeTests/Source/InterchangeTests/Private/ImportTestFunctions/SkeletalMeshImportTestFunctions.cpp
+++ b/Engine/Plugins/Tests/InterchangeTests/Source/InterchangeTests/Private/ImportTestFunctions/SkeletalMeshImportTestFunctions.cpp
@@ -4,6 +4,7 @@
 
 #include "Animation/Skeleton.h"
 #include "Engine/SkeletalMesh.h"
+#include "Engine/SkeletalMeshSocket.h"
 #include "Engine/SkinnedAssetCommon.h"
 #include "ImportTestFunctions/ImportTestFunctionsBase.h"
 #include "InterchangeTestFunction.h"
@@ -489,6 +490,52 @@ FInterchangeTestFunctionResult USkeletalMeshImportTestFunctions::CheckBonePositi
 	return Result;
 }
 
+FInterchangeTestFunctionResult USkeletalMeshImportTestFunctions::CheckSocketCount(USkeletalMesh* Mesh, int32 ExpectedSocketCount)
+{
+	FInterchangeTestFunctionResult Result;
+	if (Mesh->NumSockets() != ExpectedSocketCount)
+	{
+		Result.AddError(FString::Printf(TEXT("Expected %d sockets on the skeletal mesh, found %d"), ExpectedSocketCount, Mesh->NumSockets()));
+	}
+	return Result;
+}
+
+FInterchangeTestFunctionResult USkeletalMeshImportTestFunctions::CheckSocketName(USkeletalMesh* Mesh, int32 SocketIndex, const FString& ExpectedSocketName)
+{
+	FInterchangeTestFunctionResult Result;
+	if (USkeletalMeshSocket* Socket = Mesh->GetSocketByIndex(SocketIndex))
+	{
+		if (!ExpectedSocketName.Equals(Socket->SocketName.ToString()))
+		{
+			Result.AddError(FString::Printf(TEXT("Expected Socket with name %s at index %d for skeletal mesh, received %s"), *ExpectedSocketName, SocketIndex, *Socket->SocketName.ToString()));
+		}
+	}
+	else 
+	{
+		Result.AddError(FString::Printf(TEXT("Invalid SocketIndex(%d) for skeletal mesh"), SocketIndex));
+	}
+	return Result;
+}
+
+FInterchangeTestFunctionResult USkeletalMeshImportTestFunctions::CheckSocketLocation(USkeletalMesh* Mesh, int32 SocketIndex, const FVector& ExpectedSocketLocation)
+{
+	FInterchangeTestFunctionResult Result;
+	if (USkeletalMeshSocket* Socket = Mesh->GetSocketByIndex(SocketIndex))
+	{
+		using namespace InterchangeTestsMathUtilities;
+		const FVector RoundedSocketLocation = RoundVectorToDecimalPlaces(Socket->RelativeLocation);
+		const FVector RoundedExpectedLocation = RoundVectorToDecimalPlaces(ExpectedSocketLocation);
+		if (!RoundedExpectedLocation.Equals(RoundedSocketLocation, UE_DOUBLE_KINDA_SMALL_NUMBER))
+		{
+			Result.AddError(FString::Printf(TEXT("Expected Socket at index %d to have location %s for skeletal mesh, received %s"), SocketIndex, *RoundedSocketLocation.ToString(), *RoundedExpectedLocation.ToString()));
+		}
+	}
+	else
+	{
+		Result.AddError(FString::Printf(TEXT("Invalid SocketIndex(%d) for skeletal mesh"), SocketIndex));
+	}
+	return Result;
+}
 
 FInterchangeTestFunctionResult USkeletalMeshImportTestFunctions::CheckSkinnedVertexCountForBone(USkeletalMesh* Mesh, const FString& BoneName, bool bTestFirstAlternateProfile, int32 ExpectedSkinnedVertexCount)
 {
diff --git a/Engine/Plugins/Tests/InterchangeTests/Source/InterchangeTests/Public/ImportTestFunctions/SkeletalMeshImportTestFunctions.h b/Engine/Plugins/Tests/InterchangeTests/Source/InterchangeTests/Public/ImportTestFunctions/SkeletalMeshImportTestFunctions.h
--- a/Engine/Plugins/Tests/InterchangeTests/Source/InterchangeTests/Public/ImportTestFunctions/SkeletalMeshImportTestFunctions.h
+++ b/Engine/Plugins/Tests/InterchangeTests/Source/InterchangeTests/Public/ImportTestFunctions/SkeletalMeshImportTestFunctions.h
@@ -83,6 +83,18 @@ public:
 	UFUNCTION(Exec)
 	static FInterchangeTestFunctionResult CheckBonePosition(USkeletalMesh* Mesh, int32 BoneIndex, const FVector& ExpectedBonePosition);
 
+	/** Check whether the static mesh expected number of sockets were imported */
+	UFUNCTION(Exec)
+	static FInterchangeTestFunctionResult CheckSocketCount(USkeletalMesh* Mesh, int32 ExpectedSocketCount);
+
+	/** Check whether the static mesh given socket index has the expected name */
+	UFUNCTION(Exec)
+	static FInterchangeTestFunctionResult CheckSocketName(USkeletalMesh* Mesh, int32 SocketIndex, const FString& ExpectedSocketName);
+
+	/** Check whether the static mesh given socket index has the expected location */
+	UFUNCTION(Exec)
+	static FInterchangeTestFunctionResult CheckSocketLocation(USkeletalMesh* Mesh, int32 SocketIndex, const FVector& ExpectedSocketLocation);
+
 	/** Check that the specified bone name is skinned with the expected number of vertices */
 	UFUNCTION(Exec)
 	static FInterchangeTestFunctionResult CheckSkinnedVertexCountForBone(USkeletalMesh* Mesh, const FString& BoneName, bool bTestFirstAlternateProfile, int32 ExpectedSkinnedVertexCount);
