[Interchange] Add ability to use Camera Actor from level to capture screenshot for Test Plan screenshot comparison
- Remove the test screenshot map
- Add screenshot priorities based on the focus actor and screenshot camera in the level

#rnx

diff --git a/Engine/Plugins/Tests/InterchangeTests/Source/InterchangeTests/Private/InterchangeAutomatedTestUtils.cpp b/Engine/Plugins/Tests/InterchangeTests/Source/InterchangeTests/Private/InterchangeAutomatedTestUtils.cpp
--- a/Engine/Plugins/Tests/InterchangeTests/Source/InterchangeTests/Private/InterchangeAutomatedTestUtils.cpp
+++ b/Engine/Plugins/Tests/InterchangeTests/Source/InterchangeTests/Private/InterchangeAutomatedTestUtils.cpp
@@ -1,18 +1,20 @@
 // Copyright Epic Games, Inc. All Rights Reserved.
 #include "InterchangeAutomatedTestUtils.h"
-#include "Editor.h"
-#include "Editor/Transactor.h"
 
-#include "AutomationBlueprintFunctionLibrary.h"
-#include "Misc/AutomationTest.h"
 #include "AssetCompilingManager.h"
 #include "AssetRegistry/ARFilter.h"
 #include "AssetRegistry/AssetRegistryModule.h"
+#include "AutomationBlueprintFunctionLibrary.h"
+#include "Camera/CameraActor.h"
+#include "Editor.h"
+#include "Editor/Transactor.h"
+#include "EngineUtils.h"
 #include "HAL/FileManager.h"
 #include "InterchangeHelper.h"
 #include "InterchangeImportTestData.h"
 #include "InterchangeImportTestPlan.h"
 #include "InterchangeImportTestStepBase.h"
+#include "Misc/AutomationTest.h"
 #include "Modules/ModuleManager.h"
 #include "ObjectTools.h"
 #include "Subsystems/UnrealEditorSubsystem.h"
@@ -93,70 +95,118 @@ bool FInterchangeInterStepCollectResultCommand::Update()
 
 bool FInterchangeSetupScreenshotViewportCommand::Update()
 {
-	TArray<AActor*> ActorsInLevel;
-	const bool bUseActorFilter = !ActorClassFilter.IsEmpty();
-	for (UObject* Object : TestStepData->TestPlanData.ResultObjects)
+	if (FAutomationTestBase* CurrentTest = FAutomationTestFramework::Get().GetCurrentTest())
 	{
-		if (bUseActorFilter)
+		TestStepData->RequiredScreenshotCount++;
+		TArray<AActor*> ActorsInLevel;
+		const bool bUseActorFilter = !ScreenshotParameters.FocusActorName.IsEmpty();
+		const bool bUseActorClassFilter = ScreenshotParameters.FocusActorClass.Get() != nullptr;
+		const ACameraActor* ScreenshotCamera = []()
 		{
-			if (ActorClassFilter.Contains(Object->GetClass()))
+			ACameraActor* Camera = nullptr;
+			for (TActorIterator<ACameraActor> CameraItr(GWorld); CameraItr; ++CameraItr)
 			{
-				ActorsInLevel.Add(Cast<AActor>(Object));
+				if (*CameraItr)
+				{
+					Camera = *CameraItr;
+					break;
+				}
 			}
-		}
-		else if (Object->IsA<AActor>())
-		{
-			ActorsInLevel.Add(Cast<AActor>(Object));
-		}
-	}
-
-	if (!ActorsInLevel.IsEmpty())
-	{
-		// Cache current view mode and wireframe opacity to be restored later during cleanup.
-		TestStepData->CachedScreenshotParameters.WireframeOpacity = UAutomationBlueprintFunctionLibrary::GetEditorActiveViewportWireframeOpacity();
-		TestStepData->CachedScreenshotParameters.ViewMode = UAutomationBlueprintFunctionLibrary::GetEditorActiveViewportViewMode();
+			return Camera;
+		}();
 
-		UAutomationBlueprintFunctionLibrary::SetEditorActiveViewportViewMode(ScreenshotParameters.ViewMode);
-		UAutomationBlueprintFunctionLibrary::SetEditorActiveViewportWireframeOpacity(ScreenshotParameters.WireframeOpacity);
 
-		// Reset Viewport Camera Position to Origin for consistent results.
-		if (UUnrealEditorSubsystem* UnrealEditorSubsystem = GEditor->GetEditorSubsystem<UUnrealEditorSubsystem>())
+		for (UObject* Object : TestStepData->TestPlanData.ResultObjects)
 		{
-			UnrealEditorSubsystem->SetLevelViewportCameraInfo(FVector::ZeroVector, FRotator::ZeroRotator);
-		}
+			if (!Object->IsA<AActor>())
+			{
+				continue;
+			}
 
-		constexpr bool bActiveViewportOnly = true;
-		GEditor->MoveViewportCamerasToActor(ActorsInLevel, bActiveViewportOnly);
-	}
-	else
-	{
-		if (FAutomationTestBase* CurrentTest = FAutomationTestFramework::Get().GetCurrentTest())
-		{
+			bool bAddActor = false;
 			if (bUseActorFilter)
 			{
-				const TArray<UClass*> ActorClasses = ActorClassFilter.Array();
-				FString ClassNamesString = TEXT("");
-				for (int32 ActorClassIndex = 0; ActorClassIndex < ActorClasses.Num(); ++ActorClassIndex)
+				if (ScreenshotParameters.FocusActorName.Equals(Object->GetName()))
 				{
-					ClassNamesString += ActorClasses[ActorClassIndex]->GetName();
-					if (ActorClassIndex < ActorClasses.Num() - 1)
+					if (bUseActorClassFilter)
 					{
-						ClassNamesString += TEXT(", ");
+						if (Object && Object->GetClass()->IsChildOf(ScreenshotParameters.FocusActorClass.Get()))
+						{
+							ActorsInLevel.Add(Cast<AActor>(Object));
+						}
+						else
+						{
+							CurrentTest->AddWarning(FString::Printf(TEXT("Actor with name '%s' found but it doesn't have the actor class: %s"), *ScreenshotParameters.FocusActorName, *ScreenshotParameters.FocusActorClass.Get()->GetName()));
+						}
 					}
+					else
+					{
+						bAddActor = true;
+						if (!ActorsInLevel.IsEmpty())
+						{
+							CurrentTest->AddWarning(FString::Printf(TEXT("Current Scene contains more than one actors with the name %s. If you would like to focus on a specific actor, please consider providing the actor class in the screenshot parameters."), *Object->GetName()));
+						}
+					}
+				}
+			}
+			else 
+			{
+				bAddActor = true;
+			}
+
+			if (bAddActor)
+			{
+				ActorsInLevel.Add(Cast<AActor>(Object));
+			}
+		}
+
+		if (!ActorsInLevel.IsEmpty())
+		{
+			TestStepData->bCanTakeScreenshot = true;
+
+			// Cache current view mode and wireframe opacity to be restored later during cleanup.
+			TestStepData->CachedScreenshotParameters.WireframeOpacity = UAutomationBlueprintFunctionLibrary::GetEditorActiveViewportWireframeOpacity();
+			TestStepData->CachedScreenshotParameters.ViewMode = UAutomationBlueprintFunctionLibrary::GetEditorActiveViewportViewMode();
+
+			UAutomationBlueprintFunctionLibrary::SetEditorActiveViewportViewMode(ScreenshotParameters.ViewMode);
+			UAutomationBlueprintFunctionLibrary::SetEditorActiveViewportWireframeOpacity(ScreenshotParameters.WireframeOpacity);
+
+			// Always use camera transform when present in the level.
+			if (ScreenshotCamera)
+			{
+				if (UUnrealEditorSubsystem* UnrealEditorSubsystem = GEditor->GetEditorSubsystem<UUnrealEditorSubsystem>())
+				{
+					UnrealEditorSubsystem->SetLevelViewportCameraInfo(ScreenshotCamera->GetActorLocation(), ScreenshotCamera->GetActorRotation());
 				}
-				CurrentTest->AddWarning(FString::Printf(TEXT("Current scene does not contain any actors of classes: [%s]"), *ClassNamesString));
 			}
 			else
 			{
-				CurrentTest->AddWarning(TEXT("Current scene does not contain any actors"));
+				// Reset Viewport Camera Position to Origin for consistent results.
+				if (UUnrealEditorSubsystem* UnrealEditorSubsystem = GEditor->GetEditorSubsystem<UUnrealEditorSubsystem>())
+				{
+					UnrealEditorSubsystem->SetLevelViewportCameraInfo(FVector::ZeroVector, FRotator::ZeroRotator);
+				}
+
+				constexpr bool bActiveViewportOnly = true;
+				GEditor->MoveViewportCamerasToActor(ActorsInLevel, bActiveViewportOnly);
 			}
 		}
+		else 
+		{
+			TestStepData->bCanTakeScreenshot = false;
+			CurrentTest->AddError(TEXT("Current Scene could not focus on required actors. Screenshot will not be captured."));
+		}
 	}
 	return true;
 }
 
 bool FInterchangeCaptureScreenshotCommand::Update()
 {
+	if (!TestStepData->bCanTakeScreenshot)
+	{
+		return true;
+	}
+
 	if (TestStepData->ScreenshotTask == nullptr)
 	{
 		const FString ScreenshotNameString = FString::Printf(TEXT("%s_%d"), *TestStepData->TestPlanData.AssetData.AssetName.ToString(), TestStepData->StepIndex);
@@ -173,7 +223,13 @@ bool FInterchangeCaptureScreenshotCommand::Update()
 		return true;
 	}
 
-	return TestStepData->ScreenshotTask && TestStepData->ScreenshotTask->IsTaskDone();
+	if (TestStepData->ScreenshotTask && TestStepData->ScreenshotTask->IsTaskDone())
+	{
+		TestStepData->CapturedScreenshotCount++;
+		return true;
+	}
+
+	return false;
 }
 
 bool FInterchangeInterStepPerformTestsAndCollectGarbageCommand::Update()
@@ -186,6 +242,7 @@ bool FInterchangeInterStepPerformTestsAndCollectGarbageCommand::Update()
 
 	// Reset the Screenshot task so that next step could start a new task for screenshot.
 	TestStepData->ScreenshotTask = nullptr;
+	TestStepData->bCanTakeScreenshot = false;
 
 	TestStepData->TestPlanData.TestPlan->Steps[TestStepData->StepIndex]->FinishStep(TestStepData->TestPlanData, CurrentTest);
 	if (TestStepData->TestPlanData.InterchangeResults)
@@ -274,6 +331,7 @@ bool FInterchangeTestAutomationTestSuccessCommand::Update()
 	if (FAutomationTestBase* CurrentTest = FAutomationTestFramework::Get().GetCurrentTest())
 	{
 		CurrentTest->TestTrue(TEXT("Interchange Import Automation Test Success"), TestStepData->bSuccess);
+		CurrentTest->TestEqual(TEXT("Interchange Import Automation Test Captured Screenshots"), TestStepData->CapturedScreenshotCount, TestStepData->RequiredScreenshotCount);
 	}
 	return true;
 }
\ No newline at end of file
diff --git a/Engine/Plugins/Tests/InterchangeTests/Source/InterchangeTests/Private/InterchangeAutomatedTestUtils.h b/Engine/Plugins/Tests/InterchangeTests/Source/InterchangeTests/Private/InterchangeAutomatedTestUtils.h
--- a/Engine/Plugins/Tests/InterchangeTests/Source/InterchangeTests/Private/InterchangeAutomatedTestUtils.h
+++ b/Engine/Plugins/Tests/InterchangeTests/Source/InterchangeTests/Private/InterchangeAutomatedTestUtils.h
@@ -16,6 +16,9 @@ struct FInterchangeAutomationTestStepData
 	TTuple<UE::Interchange::FAssetImportResultPtr, UE::Interchange::FSceneImportResultPtr> Results;
 	UAutomationEditorTask* ScreenshotTask = nullptr;
 	int32 StepIndex = 0;
+	int32 RequiredScreenshotCount = 0;
+	int32 CapturedScreenshotCount = 0;
+	bool bCanTakeScreenshot = false;
 	bool bSuccess = true;
 };
 
@@ -29,7 +32,7 @@ DEFINE_LATENT_AUTOMATION_COMMAND_ONE_PARAMETER(FInterchangeInterStepPerformTests
 
 DEFINE_LATENT_AUTOMATION_COMMAND_ONE_PARAMETER(FInterchangeTestCleanUpCommand, InterchangeAutomationTestStepDataRef, TestStepData);
 
-DEFINE_LATENT_AUTOMATION_COMMAND_THREE_PARAMETER(FInterchangeSetupScreenshotViewportCommand, InterchangeAutomationTestStepDataRef, TestStepData, TSet<UClass*>, ActorClassFilter, FInterchangeTestScreenshotParameters, ScreenshotParameters);
+DEFINE_LATENT_AUTOMATION_COMMAND_TWO_PARAMETER(FInterchangeSetupScreenshotViewportCommand, InterchangeAutomationTestStepDataRef, TestStepData, FInterchangeTestScreenshotParameters, ScreenshotParameters);
 
 DEFINE_LATENT_AUTOMATION_COMMAND_TWO_PARAMETER(FInterchangeCaptureScreenshotCommand, InterchangeAutomationTestStepDataRef, TestStepData, FInterchangeTestScreenshotParameters, ScreenshotParameters);
 
diff --git a/Engine/Plugins/Tests/InterchangeTests/Source/InterchangeTests/Private/InterchangeImportTest.cpp b/Engine/Plugins/Tests/InterchangeTests/Source/InterchangeTests/Private/InterchangeImportTest.cpp
--- a/Engine/Plugins/Tests/InterchangeTests/Source/InterchangeTests/Private/InterchangeImportTest.cpp
+++ b/Engine/Plugins/Tests/InterchangeTests/Source/InterchangeTests/Private/InterchangeImportTest.cpp
@@ -127,7 +127,7 @@ bool FInterchangeImportTest::RunTest(const FString& Path)
 				const FInterchangeTestScreenshotParameters ScreenshotParameters = TestStepBase->GetScreenshotParameters();
 
 				// Focuses camera on relevant actors in the scene.
-				ADD_LATENT_AUTOMATION_COMMAND(FInterchangeSetupScreenshotViewportCommand(AutomationTestStepData.ToSharedRef(), TestStepBase->GetFocusActorsAllowedClasses(), ScreenshotParameters));
+				ADD_LATENT_AUTOMATION_COMMAND(FInterchangeSetupScreenshotViewportCommand(AutomationTestStepData.ToSharedRef(), ScreenshotParameters));
 				
 				// Camera Viewport Transition Time.
 				ADD_LATENT_AUTOMATION_COMMAND(FEngineWaitLatentCommand(0.3f));
diff --git a/Engine/Plugins/Tests/InterchangeTests/Source/InterchangeTests/Private/InterchangeImportTestStepImport.cpp b/Engine/Plugins/Tests/InterchangeTests/Source/InterchangeTests/Private/InterchangeImportTestStepImport.cpp
--- a/Engine/Plugins/Tests/InterchangeTests/Source/InterchangeTests/Private/InterchangeImportTestStepImport.cpp
+++ b/Engine/Plugins/Tests/InterchangeTests/Source/InterchangeTests/Private/InterchangeImportTestStepImport.cpp
@@ -194,19 +194,6 @@ FInterchangeTestScreenshotParameters UInterchangeImportTestStepImport::GetScreen
 	return ScreenshotParameters;
 }
 
-TSet<UClass*> UInterchangeImportTestStepImport::GetFocusActorsAllowedClasses() const
-{
-	TSet<UClass*> ActorClasses;
-	for (const TSubclassOf<AActor>& ActorSubclass : FocusActorClasses)
-	{
-		if (UClass* ActorClass = ActorSubclass.Get())
-		{
-			ActorClasses.Add(ActorClass);
-		}
-	}
-	return ActorClasses;
-}
-
 bool UInterchangeImportTestStepImport::CanEditPipelineSettings() const
 {
 	if (SourceFile.FilePath.IsEmpty())
diff --git a/Engine/Plugins/Tests/InterchangeTests/Source/InterchangeTests/Private/InterchangeImportTestStepReimport.cpp b/Engine/Plugins/Tests/InterchangeTests/Source/InterchangeTests/Private/InterchangeImportTestStepReimport.cpp
--- a/Engine/Plugins/Tests/InterchangeTests/Source/InterchangeTests/Private/InterchangeImportTestStepReimport.cpp
+++ b/Engine/Plugins/Tests/InterchangeTests/Source/InterchangeTests/Private/InterchangeImportTestStepReimport.cpp
@@ -1,6 +1,7 @@
 // Copyright Epic Games, Inc. All Rights Reserved.
 
 #include "InterchangeImportTestStepReimport.h"
+
 #include "InterchangeImportTestData.h"
 #include "Misc/PackageName.h"
 
@@ -97,11 +98,6 @@ FInterchangeTestScreenshotParameters UInterchangeImportTestStepReimport::GetScre
 	return FInterchangeTestScreenshotParameters();
 }
 
-TSet<UClass*> UInterchangeImportTestStepReimport::GetFocusActorsAllowedClasses() const
-{
-	return TSet<UClass*>();
-}
-
 bool  UInterchangeImportTestStepReimport::CanEditPipelineSettings() const
 {
 	return false;
@@ -113,4 +109,4 @@ void  UInterchangeImportTestStepReimport::EditPipelineSettings()
 
 void  UInterchangeImportTestStepReimport::ClearPipelineSettings()
 {
-}
\ No newline at end of file
+}
diff --git a/Engine/Plugins/Tests/InterchangeTests/Source/InterchangeTests/Public/InterchangeImportTestStepBase.h b/Engine/Plugins/Tests/InterchangeTests/Source/InterchangeTests/Public/InterchangeImportTestStepBase.h
--- a/Engine/Plugins/Tests/InterchangeTests/Source/InterchangeTests/Public/InterchangeImportTestStepBase.h
+++ b/Engine/Plugins/Tests/InterchangeTests/Source/InterchangeTests/Public/InterchangeImportTestStepBase.h
@@ -8,6 +8,7 @@
 #include "Engine/EngineBaseTypes.h"
 #include "InterchangeImportTestStepBase.generated.h"
 
+class ACameraActor;
 class FAutomationTestExecutionInfo;
 struct FInterchangeImportTestData;
 struct FInterchangeTestFunction;
@@ -25,6 +26,12 @@ struct FInterchangeTestScreenshotParameters
 	GENERATED_BODY()
 
 public:
+	UPROPERTY(EditAnywhere, Category = General)
+	FString FocusActorName;
+
+	UPROPERTY(EditAnywhere, Category = General)
+	TSubclassOf<AActor> FocusActorClass;
+
 	UPROPERTY(EditAnywhere, Category = General)
 	EComparisonTolerance ComparisonTolerance = EComparisonTolerance::Low;
 
@@ -53,13 +60,11 @@ public:
 	virtual FString GetContextString() const PURE_VIRTUAL(UInterchangeImportTestStepBase::GetContextString, return {}; );
 	virtual bool HasScreenshotTest() const PURE_VIRTUAL(UInterchangeImportTestStepBase::HasScreenshotTest, return false; );
 	virtual FInterchangeTestScreenshotParameters GetScreenshotParameters() const PURE_VIRTUAL(UInterchangeImportTestStepBase::GetScreenshotParameters, return FInterchangeTestScreenshotParameters{}; );
-
+	
 	virtual bool CanEditPipelineSettings() const PURE_VIRTUAL(UInterchangeImportTestStepBase::CanEditPipelineSettings, return false; );
 	virtual void EditPipelineSettings() PURE_VIRTUAL(UInterchangeImportTestStepBase::EditPipelineSettings);
 	virtual void ClearPipelineSettings() PURE_VIRTUAL(UInterchangeImportTestStepBase::ClearPipelineSettings);
 
-	virtual TSet<UClass*> GetFocusActorsAllowedClasses() const PURE_VIRTUAL(UInterchangeImportTestStepBase::GetFocusActorsAllowedClasses, return TSet<UClass*>(); );
-
 	bool PerformTests(FInterchangeImportTestData& Data, FAutomationTestBase* CurrentTest);
 protected:
 	void SaveReloadAssets(FInterchangeImportTestData& Data);
diff --git a/Engine/Plugins/Tests/InterchangeTests/Source/InterchangeTests/Public/InterchangeImportTestStepImport.h b/Engine/Plugins/Tests/InterchangeTests/Source/InterchangeTests/Public/InterchangeImportTestStepImport.h
--- a/Engine/Plugins/Tests/InterchangeTests/Source/InterchangeTests/Public/InterchangeImportTestStepImport.h
+++ b/Engine/Plugins/Tests/InterchangeTests/Source/InterchangeTests/Public/InterchangeImportTestStepImport.h
@@ -50,9 +50,6 @@ public:
 	UPROPERTY(EditAnywhere, Category = "Screenshot Comparison", meta = (AllowedClasses = "/Script/Engine.World", EditCondition = "bImportIntoLevel"))
 	FSoftObjectPath WorldPath;
 
-	UPROPERTY(EditAnywhere, Category = "Screenshot Comparison", meta = (EditCondition = "bImportIntoLevel"))
-	TArray<TSubclassOf<AActor>> FocusActorClasses;
-
 	UPROPERTY(EditAnywhere, Category = "Screenshot Comparison", meta = (EditCondition = "bImportIntoLevel"))
 	FInterchangeTestScreenshotParameters ScreenshotParameters;
 
@@ -64,7 +61,6 @@ public:
 	virtual FString GetContextString() const override;
 	virtual bool HasScreenshotTest() const override;
 	virtual FInterchangeTestScreenshotParameters GetScreenshotParameters() const override;
-	virtual TSet<UClass*> GetFocusActorsAllowedClasses() const override;
 
 	virtual bool CanEditPipelineSettings() const override;
 	virtual void EditPipelineSettings() override;
diff --git a/Engine/Plugins/Tests/InterchangeTests/Source/InterchangeTests/Public/InterchangeImportTestStepReimport.h b/Engine/Plugins/Tests/InterchangeTests/Source/InterchangeTests/Public/InterchangeImportTestStepReimport.h
--- a/Engine/Plugins/Tests/InterchangeTests/Source/InterchangeTests/Public/InterchangeImportTestStepReimport.h
+++ b/Engine/Plugins/Tests/InterchangeTests/Source/InterchangeTests/Public/InterchangeImportTestStepReimport.h
@@ -47,7 +47,6 @@ public:
 	virtual FString GetContextString() const override;
 	virtual bool HasScreenshotTest() const override;
 	virtual FInterchangeTestScreenshotParameters GetScreenshotParameters() const override;
-	virtual TSet<UClass*> GetFocusActorsAllowedClasses() const override;
 
 	virtual bool CanEditPipelineSettings() const override;
 	virtual void EditPipelineSettings() override;
