Interchange Test Plan Automated Test Plan Utils Fix
- Added a nullptr check for a UObject

#rnx

diff --git a/Engine/Plugins/Tests/InterchangeTests/Source/InterchangeTests/Private/InterchangeAutomatedTestUtils.cpp b/Engine/Plugins/Tests/InterchangeTests/Source/InterchangeTests/Private/InterchangeAutomatedTestUtils.cpp
--- a/Engine/Plugins/Tests/InterchangeTests/Source/InterchangeTests/Private/InterchangeAutomatedTestUtils.cpp
+++ b/Engine/Plugins/Tests/InterchangeTests/Source/InterchangeTests/Private/InterchangeAutomatedTestUtils.cpp
@@ -118,7 +118,7 @@ bool FInterchangeSetupScreenshotViewportCommand::Update()
 
 		for (UObject* Object : TestStepData->TestPlanData.ResultObjects)
 		{
-			if (!Object->IsA<AActor>())
+			if (!Object || (Object && !Object->IsA<AActor>()))
 			{
 				continue;
 			}
@@ -130,7 +130,7 @@ bool FInterchangeSetupScreenshotViewportCommand::Update()
 				{
 					if (bUseActorClassFilter)
 					{
-						if (Object && Object->GetClass()->IsChildOf(ScreenshotParameters.FocusActorClass.Get()))
+						if (Object->GetClass()->IsChildOf(ScreenshotParameters.FocusActorClass.Get()))
 						{
 							ActorsInLevel.Add(Cast<AActor>(Object));
 						}
