[Interchange] Test Plan Asset Screenshot Tests Camera Transform feature
- Added a way to specify the camera transform.
- Added an option to use auto focus functionality.

#rnx

diff --git a/Engine/Plugins/Tests/InterchangeTests/Source/InterchangeTests/Private/InterchangeAutomatedTestUtils.cpp b/Engine/Plugins/Tests/InterchangeTests/Source/InterchangeTests/Private/InterchangeAutomatedTestUtils.cpp
--- a/Engine/Plugins/Tests/InterchangeTests/Source/InterchangeTests/Private/InterchangeAutomatedTestUtils.cpp
+++ b/Engine/Plugins/Tests/InterchangeTests/Source/InterchangeTests/Private/InterchangeAutomatedTestUtils.cpp
@@ -99,36 +99,26 @@ bool FInterchangeSetupScreenshotViewportCommand::Update()
 	{
 		TestStepData->RequiredScreenshotCount++;
 		TArray<AActor*> ActorsInLevel;
-		const bool bUseActorFilter = !ScreenshotParameters.FocusActorName.IsEmpty();
-		const bool bUseActorClassFilter = ScreenshotParameters.FocusActorClass.Get() != nullptr;
-		const ACameraActor* ScreenshotCamera = []()
+		
+		if (ScreenshotParameters.bAutoFocus)
 		{
-			ACameraActor* Camera = nullptr;
-			for (TActorIterator<ACameraActor> CameraItr(GWorld); CameraItr; ++CameraItr)
+			for (UObject* Object : TestStepData->TestPlanData.ResultObjects)
 			{
-				if (*CameraItr)
+				if (!Object || (Object && !Object->IsA<AActor>()))
 				{
-					Camera = *CameraItr;
-					break;
+					continue;
 				}
-			}
-			return Camera;
-		}();
-
 
-		for (UObject* Object : TestStepData->TestPlanData.ResultObjects)
-		{
-			if (!Object || (Object && !Object->IsA<AActor>()))
-			{
-				continue;
-			}
+				if (ScreenshotParameters.FocusActorName.IsEmpty())
+				{
+					// No specific actor is focused, focus on the whole scene.
+					ActorsInLevel.Add(Cast<AActor>(Object));
+					continue;
+				}
 
-			bool bAddActor = false;
-			if (bUseActorFilter)
-			{
 				if (ScreenshotParameters.FocusActorName.Equals(Object->GetName()))
 				{
-					if (bUseActorClassFilter)
+					if (ScreenshotParameters.FocusActorClass.Get())
 					{
 						if (Object->GetClass()->IsChildOf(ScreenshotParameters.FocusActorClass.Get()))
 						{
@@ -141,26 +131,18 @@ bool FInterchangeSetupScreenshotViewportCommand::Update()
 					}
 					else
 					{
-						bAddActor = true;
 						if (!ActorsInLevel.IsEmpty())
 						{
 							CurrentTest->AddWarning(FString::Printf(TEXT("Current Scene contains more than one actors with the name %s. If you would like to focus on a specific actor, please consider providing the actor class in the screenshot parameters."), *Object->GetName()));
 						}
+
+						ActorsInLevel.Add(Cast<AActor>(Object));
 					}
 				}
 			}
-			else 
-			{
-				bAddActor = true;
-			}
-
-			if (bAddActor)
-			{
-				ActorsInLevel.Add(Cast<AActor>(Object));
-			}
 		}
 
-		if (!ActorsInLevel.IsEmpty())
+		if (!ScreenshotParameters.bAutoFocus || (ScreenshotParameters.bAutoFocus && !ActorsInLevel.IsEmpty()))
 		{
 			TestStepData->bCanTakeScreenshot = true;
 
@@ -171,12 +153,12 @@ bool FInterchangeSetupScreenshotViewportCommand::Update()
 			UAutomationBlueprintFunctionLibrary::SetEditorActiveViewportViewMode(ScreenshotParameters.ViewMode);
 			UAutomationBlueprintFunctionLibrary::SetEditorActiveViewportWireframeOpacity(ScreenshotParameters.WireframeOpacity);
 
-			// Always use camera transform when present in the level.
-			if (ScreenshotCamera)
+			// Always use camera transform if not using auto focus.
+			if (!ScreenshotParameters.bAutoFocus)
 			{
 				if (UUnrealEditorSubsystem* UnrealEditorSubsystem = GEditor->GetEditorSubsystem<UUnrealEditorSubsystem>())
 				{
-					UnrealEditorSubsystem->SetLevelViewportCameraInfo(ScreenshotCamera->GetActorLocation(), ScreenshotCamera->GetActorRotation());
+					UnrealEditorSubsystem->SetLevelViewportCameraInfo(ScreenshotParameters.CameraLocation, ScreenshotParameters.CameraRotation);
 				}
 			}
 			else
diff --git a/Engine/Plugins/Tests/InterchangeTests/Source/InterchangeTests/Private/InterchangeImportTestStepImport.cpp b/Engine/Plugins/Tests/InterchangeTests/Source/InterchangeTests/Private/InterchangeImportTestStepImport.cpp
--- a/Engine/Plugins/Tests/InterchangeTests/Source/InterchangeTests/Private/InterchangeImportTestStepImport.cpp
+++ b/Engine/Plugins/Tests/InterchangeTests/Source/InterchangeTests/Private/InterchangeImportTestStepImport.cpp
@@ -10,6 +10,7 @@
 #include "GameFramework/WorldSettings.h"
 #include "HAL/FileManager.h"
 #include "InterchangeImportTestData.h"
+#include "InterchangeImportTestPlan.h"
 #include "InterchangeManager.h"
 #include "InterchangeProjectSettings.h"
 #include "InterchangeTestsLog.h"
@@ -96,7 +97,7 @@ TTuple<UE::Interchange::FAssetImportResultPtr, UE::Interchange::FSceneImportResu
 	if (bImportIntoLevel)
 	{
 		// Create transient world to host data from producer
-		if (WorldPath.IsValid())
+		if (!Data.TestPlan->IsRunningSynchornously() && WorldPath.IsValid())
 		{
 			if (ULevelEditorSubsystem* LevelEditorSubsystem = GEditor->GetEditorSubsystem<ULevelEditorSubsystem>())
 			{
@@ -160,19 +161,22 @@ FTestStepResults UInterchangeImportTestStepImport::FinishStep(FInterchangeImport
 		}
 	}
 
-	if (WorldPath.IsValid())
+	if (bImportIntoLevel)
 	{
-		if (ULevelEditorSubsystem* LevelEditorSubsystem = GEditor->GetEditorSubsystem<ULevelEditorSubsystem>())
+		if (!Data.TestPlan->IsRunningSynchornously() && WorldPath.IsValid())
 		{
-			LevelEditorSubsystem->LoadLevel(WorldPath.GetAssetPathString());
+			if (ULevelEditorSubsystem* LevelEditorSubsystem = GEditor->GetEditorSubsystem<ULevelEditorSubsystem>())
+			{
+				LevelEditorSubsystem->LoadLevel(WorldPath.GetAssetPathString());
+			}
+		}
+		else if (TransientWorld)
+		{
+			// Now delete world
+			GEngine->DestroyWorldContext(TransientWorld.Get());
+			TransientWorld->DestroyWorld(true);
+			TransientWorld.Reset();
 		}
-	}
-	else if (TransientWorld)
-	{
-		// Now delete world
-		GEngine->DestroyWorldContext(TransientWorld.Get());
-		TransientWorld->DestroyWorld(true);
-		TransientWorld.Reset();
 	}
 
 	Results.bTestStepSuccess = bSuccess;
diff --git a/Engine/Plugins/Tests/InterchangeTests/Source/InterchangeTests/Public/InterchangeImportTestStepBase.h b/Engine/Plugins/Tests/InterchangeTests/Source/InterchangeTests/Public/InterchangeImportTestStepBase.h
--- a/Engine/Plugins/Tests/InterchangeTests/Source/InterchangeTests/Public/InterchangeImportTestStepBase.h
+++ b/Engine/Plugins/Tests/InterchangeTests/Source/InterchangeTests/Public/InterchangeImportTestStepBase.h
@@ -26,19 +26,28 @@ struct FInterchangeTestScreenshotParameters
 	GENERATED_BODY()
 
 public:
-	UPROPERTY(EditAnywhere, Category = General)
+	UPROPERTY(EditAnywhere, Category = Camera)
+	bool bAutoFocus = false;
+
+	UPROPERTY(EditAnywhere, Category = Camera, meta=(EditCondition="!bAutoFocus"))
+	FVector CameraLocation;
+
+	UPROPERTY(EditAnywhere, Category = Camera, meta = (EditCondition = "!bAutoFocus"))
+	FRotator CameraRotation;
+
+	UPROPERTY(EditAnywhere, Category = Camera, meta = (EditCondition = "bAutoFocus"))
 	FString FocusActorName;
 
-	UPROPERTY(EditAnywhere, Category = General)
+	UPROPERTY(EditAnywhere, Category = Camera, meta = (EditCondition = "bAutoFocus"))
 	TSubclassOf<AActor> FocusActorClass;
 
-	UPROPERTY(EditAnywhere, Category = General)
+	UPROPERTY(EditAnywhere, Category = Test)
 	EComparisonTolerance ComparisonTolerance = EComparisonTolerance::Low;
 
-	UPROPERTY(EditAnywhere, Category = General)
+	UPROPERTY(EditAnywhere, Category = Viewport)
 	TEnumAsByte<EViewModeIndex> ViewMode = EViewModeIndex::VMI_Lit;
 
-	UPROPERTY(EditAnywhere, Category = General)
+	UPROPERTY(EditAnywhere, Category = Viewport)
 	float WireframeOpacity = 0.2f;
 };
 
