[BUGFIX] Uninitialized Struct Members in FInterchangeTestScreenshotParameters

#rnx

diff --git a/Engine/Plugins/Tests/InterchangeTests/Source/InterchangeTests/Public/InterchangeImportTestStepBase.h b/Engine/Plugins/Tests/InterchangeTests/Source/InterchangeTests/Public/InterchangeImportTestStepBase.h
--- a/Engine/Plugins/Tests/InterchangeTests/Source/InterchangeTests/Public/InterchangeImportTestStepBase.h
+++ b/Engine/Plugins/Tests/InterchangeTests/Source/InterchangeTests/Public/InterchangeImportTestStepBase.h
@@ -30,10 +30,10 @@ public:
 	bool bAutoFocus = false;
 
 	UPROPERTY(EditAnywhere, Category = Camera, meta=(EditCondition="!bAutoFocus"))
-	FVector CameraLocation;
+	FVector CameraLocation = FVector::ZeroVector;
 
 	UPROPERTY(EditAnywhere, Category = Camera, meta = (EditCondition = "!bAutoFocus"))
-	FRotator CameraRotation;
+	FRotator CameraRotation = FRotator::ZeroRotator;
 
 	UPROPERTY(EditAnywhere, Category = Camera, meta = (EditCondition = "bAutoFocus"))
 	FString FocusActorName;
