[Interchange Test Plan] Skeletal Mesh Functions

- Added the ability to test vertex color by index in the test plan asset.

#rnx

diff --git a/Engine/Plugins/Tests/InterchangeTests/Source/InterchangeTests/Private/ImportTestFunctions/SkeletalMeshImportTestFunctions.cpp b/Engine/Plugins/Tests/InterchangeTests/Source/InterchangeTests/Private/ImportTestFunctions/SkeletalMeshImportTestFunctions.cpp
--- a/Engine/Plugins/Tests/InterchangeTests/Source/InterchangeTests/Private/ImportTestFunctions/SkeletalMeshImportTestFunctions.cpp
+++ b/Engine/Plugins/Tests/InterchangeTests/Source/InterchangeTests/Private/ImportTestFunctions/SkeletalMeshImportTestFunctions.cpp
@@ -383,6 +383,43 @@ FInterchangeTestFunctionResult USkeletalMeshImportTestFunctions::CheckVertexInde
 	return Result;
 }
 
+FInterchangeTestFunctionResult USkeletalMeshImportTestFunctions::CheckVertexIndexColor(USkeletalMesh* Mesh, int32 LodIndex, int32 VertexIndex, const FColor& ExpectedVertexColor)
+{
+	FInterchangeTestFunctionResult Result;
+
+	if (Mesh->GetResourceForRendering())
+	{
+		int32 ImportedLods = Mesh->GetResourceForRendering()->LODRenderData.Num();
+		if (LodIndex < 0 || LodIndex >= ImportedLods)
+		{
+			Result.AddError(FString::Printf(TEXT("The imported mesh doesn't contain LOD index %d (imported %d)."), LodIndex, ImportedLods));
+		}
+		else
+		{
+			const int32 VertexCount = Mesh->GetResourceForRendering()->LODRenderData[LodIndex].StaticVertexBuffers.ColorVertexBuffer.GetNumVertices();
+			if (VertexIndex >= VertexCount)
+			{
+				Result.AddError(FString::Printf(TEXT("The imported mesh doesn't contain vertex index %d (imported %d)."), VertexIndex, VertexCount));
+			}
+			else
+			{
+				const FColor& ImportedVertexColor = Mesh->GetResourceForRendering()->LODRenderData[LodIndex].StaticVertexBuffers.ColorVertexBuffer.VertexColor(VertexIndex);
+				if (ImportedVertexColor != ExpectedVertexColor)
+				{
+					Result.AddError(FString::Printf(TEXT("For LOD %d, vertex index %d, expected vertex color %s, imported %s."), LodIndex, VertexIndex, *ExpectedVertexColor.ToString(), *ImportedVertexColor.ToString()));
+				}
+			}
+		}
+	}
+	else
+	{
+		Result.AddError(FString::Printf(TEXT("No valid render data for this skeletalmesh %s."), *Mesh->GetName()));
+	}
+
+	return Result;
+}
+
+
 FInterchangeTestFunctionResult USkeletalMeshImportTestFunctions::CheckMorphTargetCount(USkeletalMesh* Mesh, int32 ExpectedNumberOfMorphTargets)
 {
 	FInterchangeTestFunctionResult Result;
diff --git a/Engine/Plugins/Tests/InterchangeTests/Source/InterchangeTests/Public/ImportTestFunctions/SkeletalMeshImportTestFunctions.h b/Engine/Plugins/Tests/InterchangeTests/Source/InterchangeTests/Public/ImportTestFunctions/SkeletalMeshImportTestFunctions.h
--- a/Engine/Plugins/Tests/InterchangeTests/Source/InterchangeTests/Public/ImportTestFunctions/SkeletalMeshImportTestFunctions.h
+++ b/Engine/Plugins/Tests/InterchangeTests/Source/InterchangeTests/Public/ImportTestFunctions/SkeletalMeshImportTestFunctions.h
@@ -67,6 +67,10 @@ public:
 	UFUNCTION(Exec)
 	static FInterchangeTestFunctionResult CheckVertexIndexNormal(USkeletalMesh* Mesh, int32 LodIndex, int32 VertexIndex, const FVector& ExpectedVertexNormal);
 
+	/** Check whether the vertex of the given index is at the expected normal */
+	UFUNCTION(Exec)
+	static FInterchangeTestFunctionResult CheckVertexIndexColor(USkeletalMesh* Mesh, int32 LodIndex, int32 VertexIndex, const FColor& ExpectedVertexColor);
+
 	/** Check whether the mesh has the expected number of morph targets */
 	UFUNCTION(Exec)
 	static FInterchangeTestFunctionResult CheckMorphTargetCount(USkeletalMesh* Mesh, int32 ExpectedNumberOfMorphTargets);
