[Interchange Test Plan] Object Path Substring Tests

#rnx

diff --git a/Engine/Plugins/Tests/InterchangeTests/Source/InterchangeTests/Private/ImportTestFunctions/AssetImportTestFunctions.cpp b/Engine/Plugins/Tests/InterchangeTests/Source/InterchangeTests/Private/ImportTestFunctions/AssetImportTestFunctions.cpp
--- a/Engine/Plugins/Tests/InterchangeTests/Source/InterchangeTests/Private/ImportTestFunctions/AssetImportTestFunctions.cpp
+++ b/Engine/Plugins/Tests/InterchangeTests/Source/InterchangeTests/Private/ImportTestFunctions/AssetImportTestFunctions.cpp
@@ -2,6 +2,7 @@
 
 #include "ImportTestFunctions/AssetImportTestFunctions.h"
 
+#include "AssetRegistry/AssetData.h"
 #include "ImportTestFunctions/ImportTestFunctionsBase.h"
 #include "UObject/MetaData.h"
 #include "InterchangeTestFunction.h"
@@ -66,4 +67,24 @@ FInterchangeTestFunctionResult UAssetImportTestFunctions::CheckMetadataValue(con
 	return Result;
 }
 
+FInterchangeTestFunctionResult UAssetImportTestFunctions::CheckObjectPathHasSubstring(const UObject* Object, const FString& ExpectedPathString)
+{
+	FInterchangeTestFunctionResult Result;
+	if (Object)
+	{
+		FAssetData ObjectAssetData(Object);
+		const FString ObjectPathString = ObjectAssetData.GetObjectPathString();
+		const int32 FoundIndex = ObjectPathString.Find(ExpectedPathString);
+		if (FoundIndex < 0)
+		{
+			Result.AddError(FString::Printf(TEXT("Expected Object Path to contain [%s], but no matching substring was found in the object path [%s]. "), *ExpectedPathString, *ObjectPathString));
+		}
+	}
+	else
+	{
+		Result.AddError(TEXT("Can't retrieve paths invalid object."));
+	}
+	return Result;
+}
+
 
diff --git a/Engine/Plugins/Tests/InterchangeTests/Source/InterchangeTests/Public/ImportTestFunctions/AssetImportTestFunctions.h b/Engine/Plugins/Tests/InterchangeTests/Source/InterchangeTests/Public/ImportTestFunctions/AssetImportTestFunctions.h
--- a/Engine/Plugins/Tests/InterchangeTests/Source/InterchangeTests/Public/ImportTestFunctions/AssetImportTestFunctions.h
+++ b/Engine/Plugins/Tests/InterchangeTests/Source/InterchangeTests/Public/ImportTestFunctions/AssetImportTestFunctions.h
@@ -39,4 +39,11 @@ public:
 	UFUNCTION(Exec)
 	static FInterchangeTestFunctionResult CheckMetadataValue(const UObject* Object, const FString& ExpectedMetadataKey, const FString& ExpectedMetadataValue);
 
+	/**
+	 * Check whether the object's path has the given string as substring.
+	 * @Param ExpectedPathString - The string to look for in the object path.
+	 */
+	UFUNCTION(Exec)
+	static FInterchangeTestFunctionResult CheckObjectPathHasSubstring(const UObject* Object, const FString& ExpectedPathString);
+
 };
