[Interchange] Migrate Test Plan to the new design
- Test Plan now has one import step and multiple reimport steps
- SaveReloadAssets feature is removed as the same thing can be emulated with multiple test plan assets. Can be introduced again later.

[Interchange] Create a description field in the Interchange Import Test Plan
- Added a description field to the Interchange Test Plan

[BUGFIX] Interchange test plan - Level import - Instance and packed hierarchy test failed (Crash)

#rnx

diff --git a/Engine/Plugins/Interchange/Editor/Source/Pipelines/Private/SInterchangePipelineConfigurationDialog.cpp b/Engine/Plugins/Interchange/Editor/Source/Pipelines/Private/SInterchangePipelineConfigurationDialog.cpp
--- a/Engine/Plugins/Interchange/Editor/Source/Pipelines/Private/SInterchangePipelineConfigurationDialog.cpp
+++ b/Engine/Plugins/Interchange/Editor/Source/Pipelines/Private/SInterchangePipelineConfigurationDialog.cpp
@@ -329,7 +329,7 @@ TSharedRef<SBox> SInterchangePipelineConfigurationDialog::SpawnPipelineConfigura
 				if (UInterchangePipelineBase* GeneratedPipeline = UE::Interchange::GeneratePipelineInstance(DefaultPipeline))
 				{
 					GeneratedPipeline->TransferAdjustSettings(DefaultPipeline);
-					if (bTestConfigDialog || GeneratedPipeline->IsFromReimportOrOverride())
+					if (GeneratedPipeline->IsFromReimportOrOverride())
 					{
 						//We save the pipeline settings to allow Reset to Default to work
 						GeneratedPipeline->SaveSettings(Stack.StackName);
diff --git a/Engine/Plugins/Tests/InterchangeTests/Source/InterchangeTestEditor/Private/InterchangeImportTestPlanAssetDetails.cpp b/Engine/Plugins/Tests/InterchangeTests/Source/InterchangeTestEditor/Private/InterchangeImportTestPlanAssetDetails.cpp
--- a/Engine/Plugins/Tests/InterchangeTests/Source/InterchangeTestEditor/Private/InterchangeImportTestPlanAssetDetails.cpp
+++ b/Engine/Plugins/Tests/InterchangeTests/Source/InterchangeTestEditor/Private/InterchangeImportTestPlanAssetDetails.cpp
@@ -19,13 +19,11 @@ FInterchangeImportTestPlanAssetDetailsCustomization::FInterchangeImportTestPlanA
 {
 }
 
-
 TSharedRef<IDetailCustomization> FInterchangeImportTestPlanAssetDetailsCustomization::MakeInstance()
 {
 	return MakeShareable(new FInterchangeImportTestPlanAssetDetailsCustomization);
 }
 
-
 void FInterchangeImportTestPlanAssetDetailsCustomization::CustomizeDetails(IDetailLayoutBuilder& DetailBuilder)
 {
 	CachedDetailBuilder = &DetailBuilder;
@@ -44,28 +42,32 @@ void FInterchangeImportTestPlanAssetDetailsCustomization::CustomizeDetails(IDeta
 		return;
 	}
 
+	CustomizeAutomationCategory();
+}
 
-	TFunction<bool(const UInterchangeImportTestPlan* InTestPlanAsset, FString& OutReason)> IsTestPlanAssetSkipped = [](const UInterchangeImportTestPlan * InTestPlanAsset, FString & OutReason)
-	{
-		if (!InTestPlanAsset)
+void FInterchangeImportTestPlanAssetDetailsCustomization::CustomizeAutomationCategory()
+{
+	TFunction<bool(const UInterchangeImportTestPlan* InTestPlanAsset, FString& OutReason)> IsTestPlanAssetSkipped = [](const UInterchangeImportTestPlan* InTestPlanAsset, FString& OutReason)
 		{
-			return false;
-		}
+			if (!InTestPlanAsset)
+			{
+				return false;
+			}
 
-		if (UAutomationTestExcludelist* AutomationTestExcludeList = UAutomationTestExcludelist::Get())
-		{
-			FAssetData TestPlanAssetData(InTestPlanAsset);
-			constexpr bool bAddBeautifiedTestNamePrefix = true;
-			const FString TestName = UE::Interchange::FInterchangeImportTestPlanStaticHelpers::GetTestNameFromObjectPathString(TestPlanAssetData.GetObjectPathString(), bAddBeautifiedTestNamePrefix);
-			if (const FAutomationTestExcludelistEntry* TestEntry = AutomationTestExcludeList->GetExcludeTestEntry(TestName))
+			if (UAutomationTestExcludelist* AutomationTestExcludeList = UAutomationTestExcludelist::Get())
 			{
-				OutReason = TestEntry->Reason.ToString();
-				return true;
+				FAssetData TestPlanAssetData(InTestPlanAsset);
+				constexpr bool bAddBeautifiedTestNamePrefix = true;
+				const FString TestName = UE::Interchange::FInterchangeImportTestPlanStaticHelpers::GetTestNameFromObjectPathString(TestPlanAssetData.GetObjectPathString(), bAddBeautifiedTestNamePrefix);
+				if (const FAutomationTestExcludelistEntry* TestEntry = AutomationTestExcludeList->GetExcludeTestEntry(TestName))
+				{
+					OutReason = TestEntry->Reason.ToString();
+					return true;
+				}
 			}
-		}
 
-		return false;
-	};
+			return false;
+		};
 
 
 	IDetailCategoryBuilder& AutomationCategory = CachedDetailBuilder->EditCategory(FName("Automation"), FText::GetEmpty(), ECategoryPriority::Uncommon);
@@ -80,12 +82,12 @@ void FInterchangeImportTestPlanAssetDetailsCustomization::CustomizeDetails(IDeta
 		.ValueContent()
 		[
 			SNew(SHorizontalBox)
-			+SHorizontalBox::Slot()
-			.AutoWidth()
-			[
-				SNew(SCheckBox)
-					.IsEnabled(false)
-					.IsChecked_Lambda([this, IsTestPlanAssetSkipped]() {
+				+ SHorizontalBox::Slot()
+				.AutoWidth()
+				[
+					SNew(SCheckBox)
+						.IsEnabled(false)
+						.IsChecked_Lambda([this, IsTestPlanAssetSkipped]() {
 						if (InterchangeImportTestPlan.IsValid())
 						{
 							FString OutReason;
@@ -96,42 +98,41 @@ void FInterchangeImportTestPlanAssetDetailsCustomization::CustomizeDetails(IDeta
 						}
 
 						return ECheckBoxState::Unchecked;
-					})
-			]
-			+SHorizontalBox::Slot()
-			.AutoWidth()
-			.Padding(2.0, 4.0)
-			[
-				SNew(STextBlock)
-				.Font(FAppStyle::GetFontStyle(TEXT("PropertyWindow.NormalFont")))
-				.Visibility_Lambda([this, IsTestPlanAssetSkipped]()
-					{
-						if (InterchangeImportTestPlan.IsValid())
-						{
-							FString OutReason;
-							if (IsTestPlanAssetSkipped(InterchangeImportTestPlan.Get(), OutReason))
+							})
+				]
+				+ SHorizontalBox::Slot()
+				.AutoWidth()
+				.Padding(2.0, 4.0)
+				[
+					SNew(STextBlock)
+						.Font(FAppStyle::GetFontStyle(TEXT("PropertyWindow.NormalFont")))
+						.Visibility_Lambda([this, IsTestPlanAssetSkipped]()
 							{
-								return EVisibility::Visible;
-							}
-						}
-
-						return EVisibility::Collapsed;
-					})
-				.Text_Lambda([this, IsTestPlanAssetSkipped]()
-					{
-						if (InterchangeImportTestPlan.IsValid())
-						{
-							FString OutReason;
-							if (IsTestPlanAssetSkipped(InterchangeImportTestPlan.Get(), OutReason))
+								if (InterchangeImportTestPlan.IsValid())
+								{
+									FString OutReason;
+									if (IsTestPlanAssetSkipped(InterchangeImportTestPlan.Get(), OutReason))
+									{
+										return EVisibility::Visible;
+									}
+								}
+
+								return EVisibility::Collapsed;
+							})
+						.Text_Lambda([this, IsTestPlanAssetSkipped]()
 							{
-								return FText::FromString(FString::Printf(TEXT("Reason Skipped: %s"), *OutReason));
-							}
-						}
-						return FText::GetEmpty();
-					})
-			]
+								if (InterchangeImportTestPlan.IsValid())
+								{
+									FString OutReason;
+									if (IsTestPlanAssetSkipped(InterchangeImportTestPlan.Get(), OutReason))
+									{
+										return FText::FromString(FString::Printf(TEXT("Reason Skipped: %s"), *OutReason));
+									}
+								}
+								return FText::GetEmpty();
+							})
+				]
 		];
 }
 
-
 #undef LOCTEXT_NAMESPACE
\ No newline at end of file
diff --git a/Engine/Plugins/Tests/InterchangeTests/Source/InterchangeTestEditor/Private/InterchangeImportTestPlanAssetDetails.h b/Engine/Plugins/Tests/InterchangeTests/Source/InterchangeTestEditor/Private/InterchangeImportTestPlanAssetDetails.h
--- a/Engine/Plugins/Tests/InterchangeTests/Source/InterchangeTestEditor/Private/InterchangeImportTestPlanAssetDetails.h
+++ b/Engine/Plugins/Tests/InterchangeTests/Source/InterchangeTestEditor/Private/InterchangeImportTestPlanAssetDetails.h
@@ -21,6 +21,7 @@ public:
 private:
 	FInterchangeImportTestPlanAssetDetailsCustomization();
 
+	void CustomizeAutomationCategory();
 private:
 	TWeakObjectPtr<UInterchangeImportTestPlan> InterchangeImportTestPlan;
 
diff --git a/Engine/Plugins/Tests/InterchangeTests/Source/InterchangeTestEditor/Private/InterchangeTestPlanPipelineSettingsLayout.cpp b/Engine/Plugins/Tests/InterchangeTests/Source/InterchangeTestEditor/Private/InterchangeTestPlanPipelineSettingsLayout.cpp
--- a/Engine/Plugins/Tests/InterchangeTests/Source/InterchangeTestEditor/Private/InterchangeTestPlanPipelineSettingsLayout.cpp
+++ b/Engine/Plugins/Tests/InterchangeTests/Source/InterchangeTestEditor/Private/InterchangeTestPlanPipelineSettingsLayout.cpp
@@ -3,6 +3,7 @@
 #include "InterchangeTestPlanPipelineSettings.h"
 #include "InterchangeImportTestStepBase.h"
 #include "InterchangeImportTestStepImport.h"
+#include "InterchangeImportTestStepReimport.h"
 #include "InterchangePipelineBase.h"
 
 #include "DetailWidgetRow.h"
@@ -12,8 +13,11 @@
 #include "IDetailChildrenBuilder.h"
 #include "IPropertyUtilities.h"
 #include "Misc/ConfigCacheIni.h"
+#include "Widgets/Images/SImage.h"
 #include "Widgets/Input/SButton.h"
-#include "Widgets/Input/SCheckBox.h"
+#include "Widgets/Layout/SBox.h"
+#include "Widgets/SBoxPanel.h"
+#include "Widgets/Text/STextBlock.h"
 
 #define LOCTEXT_NAMESPACE "InterchangeTestPipelineSettings"
 
@@ -85,6 +89,44 @@ void FInterchangeTestPlanPipelineSettingsLayout::CustomizeHeader(TSharedRef<IPro
 						})
 					.OnClicked(this, &FInterchangeTestPlanPipelineSettingsLayout::ClearModifiedPipelineSettings)
 			]
+			+ SHorizontalBox::Slot()
+			.Padding(4.0, 2.0)
+			.HAlign(HAlign_Center)
+			.VAlign(VAlign_Center)
+			.AutoWidth()
+			[
+				SNew(STextBlock)
+					.Visibility_Lambda([this]() 
+					{
+						if (FInterchangeTestPlanPipelineSettings* PipelineSettings = GetStruct())
+						{
+							if (UInterchangeImportTestStepBase* TestStep = PipelineSettings->ParentTestStep.Get())
+							{
+								constexpr bool bCheckForValidPipelines = false;
+								if (!TestStep->IsUsingOverridePipelines(bCheckForValidPipelines))
+								{
+									return EVisibility::Visible;
+								}
+							}
+						}
+						return EVisibility::Collapsed;
+					})
+					.Text_Lambda([this]()
+					{
+						if (FInterchangeTestPlanPipelineSettings* PipelineSettings = GetStruct())
+						{
+							if (PipelineSettings->CustomPipelines.IsEmpty())
+							{
+								return LOCTEXT("PipelineSettingsCountEmptyText", "Pipeline Count : Empty");
+							}
+
+							return FText::Format(LOCTEXT("PipelineSettingsCountText", "Pipeline Count : {0}"), PipelineSettings->CustomPipelines.Num());
+						}
+
+						return FText::GetEmpty();
+					})
+					.Font(IDetailLayoutBuilder::GetDetailFont())
+			]
 	];
 }
 
diff --git a/Engine/Plugins/Tests/InterchangeTests/Source/InterchangeTests/Private/InterchangeAutomatedTestUtils.cpp b/Engine/Plugins/Tests/InterchangeTests/Source/InterchangeTests/Private/InterchangeAutomatedTestUtils.cpp
--- a/Engine/Plugins/Tests/InterchangeTests/Source/InterchangeTests/Private/InterchangeAutomatedTestUtils.cpp
+++ b/Engine/Plugins/Tests/InterchangeTests/Source/InterchangeTests/Private/InterchangeAutomatedTestUtils.cpp
@@ -14,16 +14,80 @@
 #include "InterchangeImportTestData.h"
 #include "InterchangeImportTestPlan.h"
 #include "InterchangeImportTestStepBase.h"
+#include "InterchangeImportTestStepImport.h"
+#include "InterchangeImportTestStepReimport.h"
 #include "Misc/AutomationTest.h"
 #include "Modules/ModuleManager.h"
 #include "ObjectTools.h"
 #include "Subsystems/UnrealEditorSubsystem.h"
 
+namespace InterchangeAutomationTestUtils::Private
+{
+	TTuple<UE::Interchange::FAssetImportResultPtr, UE::Interchange::FSceneImportResultPtr> StartCurrentStep(FInterchangeAutomationTestStepData& TestStepData)
+	{
+		if (TestStepData.bIsReimportStep)
+		{
+			if (ensure(TestStepData.TestPlanData.TestPlan->ReimportStack.IsValidIndex(TestStepData.StepIndex)))
+			{
+				return TestStepData.TestPlanData.TestPlan->ReimportStack[TestStepData.StepIndex]->StartStep(TestStepData.TestPlanData);
+			}
+			
+			return { nullptr, nullptr };
+		}
+		else
+		{
+			return TestStepData.TestPlanData.TestPlan->ImportStep->StartStep(TestStepData.TestPlanData);
+		}
+	}
+
+	void FinishCurrentStep(FInterchangeAutomationTestStepData& TestStepData, FAutomationTestBase* CurrentTest)
+	{
+		if (TestStepData.bIsReimportStep)
+		{
+			TestStepData.TestPlanData.TestPlan->ReimportStack[TestStepData.StepIndex]->FinishStep(TestStepData.TestPlanData, CurrentTest);
+		}
+		else
+		{
+			TestStepData.TestPlanData.TestPlan->ImportStep->FinishStep(TestStepData.TestPlanData, CurrentTest);
+		}
+	}
+
+	FString GetContextString(const FInterchangeAutomationTestStepData& TestStepData)
+	{
+		if (TestStepData.bIsReimportStep)
+		{
+			return TestStepData.TestPlanData.AssetName.ToString() + TEXT(": ") + TestStepData.TestPlanData.TestPlan->ReimportStack[TestStepData.StepIndex]->GetContextString();
+		}
+		else
+		{
+			return TestStepData.TestPlanData.AssetName.ToString() + TEXT(": ") + TestStepData.TestPlanData.TestPlan->ImportStep->GetContextString();
+		}
+	}
+
+	FString GetScreenshotNameString(const FInterchangeAutomationTestStepData& TestStepData)
+	{
+		if (TestStepData.bIsReimportStep)
+		{
+			return FString::Printf(TEXT("%s_ReimportStep_%d"), *TestStepData.TestPlanData.AssetName.ToString(), TestStepData.StepIndex);
+		}
+		else
+		{
+			return FString::Printf(TEXT("%s_ImportStep"), *TestStepData.TestPlanData.AssetName.ToString());
+		}
+	}
+}
+
+
 bool FInterchangeIntializeStepCommand::Update()
 {
+	using namespace InterchangeAutomationTestUtils::Private;
+
+	TGuardValue<bool> UnattendedScriptGuard(GIsRunningUnattendedScript, true);
+	TestStepData->bIsReimportStep = bIsReimport;
 	TestStepData->StepIndex = StepIndex;
 	TestStepData->TestPlanData.ImportedAssets.Empty();
-	TestStepData->Results = TestStepData->TestPlanData.TestPlan->Steps[TestStepData->StepIndex]->StartStep(TestStepData->TestPlanData);
+	TestStepData->Results = StartCurrentStep(*TestStepData);
+
 	if (UE::Interchange::FAssetImportResultPtr AssetImportResults = TestStepData->Results.Get<0>())
 	{
 		AssetImportResults->WaitUntilDone();
@@ -39,10 +103,12 @@ bool FInterchangeIntializeStepCommand::Update()
 
 bool FInterchangeInterStepCollectResultCommand::Update()
 {
+	using namespace InterchangeAutomationTestUtils::Private;
 	if (FAutomationTestBase* CurrentTest = FAutomationTestFramework::Get().GetCurrentTest())
 	{
-		FString Context = TestStepData->TestPlanData.AssetData.AssetName.ToString() + TEXT(": ") + TestStepData->TestPlanData.TestPlan->Steps[TestStepData->StepIndex]->GetContextString();
-		CurrentTest->PushContext(Context);
+		const FString ContextString = GetContextString(*TestStepData);
+
+		CurrentTest->PushContext(ContextString);
 
 		// Fill out list of result objects in the data object.
 		// These are the UObject* results corresponding to imported assets.
@@ -89,6 +155,24 @@ bool FInterchangeInterStepCollectResultCommand::Update()
 		{
 			TestStepData->TestPlanData.InterchangeResults->SetFlags(RF_Standalone);
 		}
+
+		{
+			for (int32 ResultObjectIndex = 0; ResultObjectIndex < TestStepData->TestPlanData.ResultObjects.Num(); ++ResultObjectIndex)
+			{
+				UObject* ResultObject = TestStepData->TestPlanData.ResultObjects[ResultObjectIndex];
+				if (UWorld* World = Cast<UWorld>(ResultObject))
+				{
+					if (TestStepData->TestPlanData.TestPlan->GetCurrentWorld() == World)
+					{
+						continue;
+					}
+					else
+					{
+						ResultObject->SetFlags(RF_Standalone);
+					}
+				}
+			}
+		}
 	}
 	return true;
 }
@@ -99,7 +183,7 @@ bool FInterchangeSetupScreenshotViewportCommand::Update()
 	{
 		TestStepData->RequiredScreenshotCount++;
 		TArray<AActor*> ActorsInLevel;
-		
+
 		if (ScreenshotParameters.bAutoFocus)
 		{
 			for (UObject* Object : TestStepData->TestPlanData.ResultObjects)
@@ -173,7 +257,7 @@ bool FInterchangeSetupScreenshotViewportCommand::Update()
 				GEditor->MoveViewportCamerasToActor(ActorsInLevel, bActiveViewportOnly);
 			}
 		}
-		else 
+		else
 		{
 			TestStepData->bCanTakeScreenshot = false;
 			CurrentTest->AddError(TEXT("Current Scene could not focus on required actors. Screenshot will not be captured."));
@@ -191,10 +275,11 @@ bool FInterchangeCaptureScreenshotCommand::Update()
 
 	if (TestStepData->ScreenshotTask == nullptr)
 	{
-		const FString ScreenshotNameString = FString::Printf(TEXT("%s_%d"), *TestStepData->TestPlanData.AssetData.AssetName.ToString(), TestStepData->StepIndex);
-		TestStepData->ScreenshotTask = UAutomationBlueprintFunctionLibrary::TakeHighResScreenshot(1280, 720, ScreenshotNameString, nullptr, false, false, ScreenshotParameters.ComparisonTolerance);
+		using namespace InterchangeAutomationTestUtils::Private;
+
+		TestStepData->ScreenshotTask = UAutomationBlueprintFunctionLibrary::TakeHighResScreenshot(1280, 720, GetScreenshotNameString(*TestStepData), nullptr, false, false, ScreenshotParameters.ComparisonTolerance);
 	}
-	
+
 	if (TestStepData->ScreenshotTask && !TestStepData->ScreenshotTask->IsValidTask())
 	{
 		if (FAutomationTestBase* CurrentTest = FAutomationTestFramework::Get().GetCurrentTest())
@@ -216,6 +301,8 @@ bool FInterchangeCaptureScreenshotCommand::Update()
 
 bool FInterchangeInterStepPerformTestsAndCollectGarbageCommand::Update()
 {
+	using namespace InterchangeAutomationTestUtils::Private;
+
 	FAutomationTestBase* CurrentTest = FAutomationTestFramework::Get().GetCurrentTest();
 
 	// Restore Viewport Parameters
@@ -226,7 +313,8 @@ bool FInterchangeInterStepPerformTestsAndCollectGarbageCommand::Update()
 	TestStepData->ScreenshotTask = nullptr;
 	TestStepData->bCanTakeScreenshot = false;
 
-	TestStepData->TestPlanData.TestPlan->Steps[TestStepData->StepIndex]->FinishStep(TestStepData->TestPlanData, CurrentTest);
+	FinishCurrentStep(*TestStepData, CurrentTest);
+
 	if (TestStepData->TestPlanData.InterchangeResults)
 	{
 		// Populate the automation test execution info with the interchange import results
@@ -271,8 +359,52 @@ bool FInterchangeInterStepPerformTestsAndCollectGarbageCommand::Update()
 	{
 		GEditor->Trans->Reset(FText::FromString("Discard undo history during Automation testing."));
 	}
+
+	{
+		TArray<int32> RemoveObjectIndices;
+		for (int32 ResultObjectIndex = 0; ResultObjectIndex < TestStepData->TestPlanData.ResultObjects.Num(); ++ResultObjectIndex)
+		{
+			UObject* ResultObject = TestStepData->TestPlanData.ResultObjects[ResultObjectIndex];
+
+			if (UWorld* World = Cast<UWorld>(ResultObject))
+			{
+				if (TestStepData->TestPlanData.TestPlan->GetCurrentWorld() == World)
+				{
+					continue;
+				}
+				else
+				{
+					ResultObject->ClearFlags(RF_Standalone);
+					RemoveObjectIndices.Add(ResultObjectIndex);
+				}
+			}
+		}
+
+		for (int32 RemoveIndex = RemoveObjectIndices.Num() - 1; RemoveIndex >= 0; --RemoveIndex)
+		{
+			TestStepData->TestPlanData.ResultObjects.RemoveAt(RemoveObjectIndices[RemoveIndex]);
+		}
+	}
+
 	CollectGarbage(GARBAGE_COLLECTION_KEEPFLAGS);
 
+	{
+		TArray<int32> RemoveObjectIndices;
+		for (int32 ResultObjectIndex = 0; ResultObjectIndex < TestStepData->TestPlanData.ResultObjects.Num(); ++ResultObjectIndex)
+		{
+			UObject* ResultObject = TestStepData->TestPlanData.ResultObjects[ResultObjectIndex];
+			if (ResultObject && ResultObject->HasAllFlags(RF_BeginDestroyed | RF_FinishDestroyed))
+			{
+				RemoveObjectIndices.Add(ResultObjectIndex);
+			}
+		}
+
+		for (int32 RemoveIndex = RemoveObjectIndices.Num() - 1; RemoveIndex >= 0; --RemoveIndex)
+		{
+			TestStepData->TestPlanData.ResultObjects.RemoveAt(RemoveObjectIndices[RemoveIndex]);
+		}
+	}
+
 	return true;
 }
 
@@ -280,10 +412,32 @@ bool FInterchangeInterStepPerformTestsAndCollectGarbageCommand::Update()
 bool FInterchangeTestCleanUpCommand::Update()
 {
 	TArray<UObject*> ObjectsToDelete;
+	TArray<UWorld*> WorldsToDelete;
 	ObjectsToDelete.Reserve(ObjectsToDelete.Max() + TestStepData->TestPlanData.ResultObjects.Num());
 	for (UObject* ResultObject : TestStepData->TestPlanData.ResultObjects)
 	{
-		if (AActor* Actor = Cast<AActor>(ResultObject))
+		if (!ResultObject)
+		{
+			continue;
+		}
+
+		if (ResultObject->HasAnyFlags(RF_BeginDestroyed | RF_FinishDestroyed))
+		{
+			continue;
+		}
+
+		if (UWorld* World = Cast<UWorld>(ResultObject))
+		{
+			if (TestStepData->TestPlanData.TestPlan->GetCurrentWorld() == World)
+			{
+				continue;
+			}
+			else
+			{
+				WorldsToDelete.Add(World);
+			}
+		}
+		else if (AActor* Actor = Cast<AActor>(ResultObject))
 		{
 			constexpr bool bShouldModifyLevel = true;
 			Actor->GetWorld()->EditorDestroyActor(Actor, bShouldModifyLevel);
@@ -292,22 +446,39 @@ bool FInterchangeTestCleanUpCommand::Update()
 		}
 		else
 		{
-			if (ResultObject && !ResultObject->HasAnyFlags(RF_BeginDestroyed | RF_FinishDestroyed))
-			{
-				ObjectsToDelete.Add(ResultObject);
-			}
+			ObjectsToDelete.Add(ResultObject);
 		}
 	}
+
 	//Make sure all compilation is done before deleting some objects
 	FAssetCompilingManager::Get().FinishAllCompilation();
 
+	for (UWorld* WorldToDelete : WorldsToDelete)
+	{
+		if (!WorldToDelete->HasAnyFlags(RF_BeginDestroyed | RF_FinishDestroyed))
+		{
+			GEngine->DestroyWorldContext(WorldToDelete);
+			WorldToDelete->DestroyWorld(true);
+		}
+	}
+
 	constexpr bool bShowConfirmation = false;
 	ObjectTools::ForceDeleteObjects(ObjectsToDelete, bShowConfirmation);
+
+	// Destroy Transient World or reload the custom level
+	TestStepData->TestPlanData.TestPlan->CleanupLevel();
+
+	if (!TestStepData->PendingDeleteDirectoryPath.IsEmpty())
+	{
+		constexpr bool bRequireExists = false;
+		constexpr bool bDeleteRecursively = true;
+		IFileManager::Get().DeleteDirectory(*TestStepData->PendingDeleteDirectoryPath, bRequireExists, bDeleteRecursively);
+	}
+
 	return true;
 
 }
 
-
 bool FInterchangeTestAutomationTestSuccessCommand::Update()
 {
 	if (FAutomationTestBase* CurrentTest = FAutomationTestFramework::Get().GetCurrentTest())
diff --git a/Engine/Plugins/Tests/InterchangeTests/Source/InterchangeTests/Private/InterchangeAutomatedTestUtils.h b/Engine/Plugins/Tests/InterchangeTests/Source/InterchangeTests/Private/InterchangeAutomatedTestUtils.h
--- a/Engine/Plugins/Tests/InterchangeTests/Source/InterchangeTests/Private/InterchangeAutomatedTestUtils.h
+++ b/Engine/Plugins/Tests/InterchangeTests/Source/InterchangeTests/Private/InterchangeAutomatedTestUtils.h
@@ -11,6 +11,7 @@ class UClass;
 
 struct FInterchangeAutomationTestStepData
 {
+	FString PendingDeleteDirectoryPath; 
 	FInterchangeImportTestData TestPlanData;
 	FInterchangeTestScreenshotParameters CachedScreenshotParameters;
 	TTuple<UE::Interchange::FAssetImportResultPtr, UE::Interchange::FSceneImportResultPtr> Results;
@@ -18,22 +19,23 @@ struct FInterchangeAutomationTestStepData
 	int32 StepIndex = 0;
 	int32 RequiredScreenshotCount = 0;
 	int32 CapturedScreenshotCount = 0;
+	bool bIsReimportStep = false;
 	bool bCanTakeScreenshot = false;
 	bool bSuccess = true;
 };
 
 using InterchangeAutomationTestStepDataRef = TSharedRef< FInterchangeAutomationTestStepData, ESPMode::ThreadSafe>;
 
-DEFINE_LATENT_AUTOMATION_COMMAND_TWO_PARAMETER(FInterchangeIntializeStepCommand, InterchangeAutomationTestStepDataRef, TestStepData, int32, StepIndex);
+DEFINE_LATENT_AUTOMATION_COMMAND_THREE_PARAMETER(FInterchangeIntializeStepCommand, InterchangeAutomationTestStepDataRef, TestStepData, bool, bIsReimport, int32, StepIndex);
 
 DEFINE_LATENT_AUTOMATION_COMMAND_ONE_PARAMETER(FInterchangeInterStepCollectResultCommand, InterchangeAutomationTestStepDataRef, TestStepData);
 
-DEFINE_LATENT_AUTOMATION_COMMAND_ONE_PARAMETER(FInterchangeInterStepPerformTestsAndCollectGarbageCommand, InterchangeAutomationTestStepDataRef, TestStepData);
-
-DEFINE_LATENT_AUTOMATION_COMMAND_ONE_PARAMETER(FInterchangeTestCleanUpCommand, InterchangeAutomationTestStepDataRef, TestStepData);
-
 DEFINE_LATENT_AUTOMATION_COMMAND_TWO_PARAMETER(FInterchangeSetupScreenshotViewportCommand, InterchangeAutomationTestStepDataRef, TestStepData, FInterchangeTestScreenshotParameters, ScreenshotParameters);
 
 DEFINE_LATENT_AUTOMATION_COMMAND_TWO_PARAMETER(FInterchangeCaptureScreenshotCommand, InterchangeAutomationTestStepDataRef, TestStepData, FInterchangeTestScreenshotParameters, ScreenshotParameters);
 
+DEFINE_LATENT_AUTOMATION_COMMAND_ONE_PARAMETER(FInterchangeInterStepPerformTestsAndCollectGarbageCommand, InterchangeAutomationTestStepDataRef, TestStepData);
+
 DEFINE_LATENT_AUTOMATION_COMMAND_ONE_PARAMETER(FInterchangeTestAutomationTestSuccessCommand, InterchangeAutomationTestStepDataRef, TestStepData);
+
+DEFINE_LATENT_AUTOMATION_COMMAND_ONE_PARAMETER(FInterchangeTestCleanUpCommand, InterchangeAutomationTestStepDataRef, TestStepData);
\ No newline at end of file
diff --git a/Engine/Plugins/Tests/InterchangeTests/Source/InterchangeTests/Private/InterchangeImportTest.cpp b/Engine/Plugins/Tests/InterchangeTests/Source/InterchangeTests/Private/InterchangeImportTest.cpp
--- a/Engine/Plugins/Tests/InterchangeTests/Source/InterchangeTests/Private/InterchangeImportTest.cpp
+++ b/Engine/Plugins/Tests/InterchangeTests/Source/InterchangeTests/Private/InterchangeImportTest.cpp
@@ -11,10 +11,12 @@
 #include "InterchangeImportTestData.h"
 #include "InterchangeImportTestPlan.h"
 #include "InterchangeImportTestStepBase.h"
+#include "InterchangeImportTestStepImport.h"
+#include "InterchangeImportTestStepReimport.h"
 #include "Misc/AutomationTest.h"
-#include "Tests/AutomationCommon.h"
 #include "Modules/ModuleManager.h"
 #include "ObjectTools.h"
+#include "Tests/AutomationCommon.h"
 
 IMPLEMENT_CUSTOM_COMPLEX_AUTOMATION_TEST(FInterchangeImportTest, FAutomationTestBase, UE::Interchange::FInterchangeImportTestPlanStaticHelpers::GetBeautifiedTestName(), EAutomationTestFlags::EditorContext | EAutomationTestFlags::EngineFilter)
 
@@ -76,12 +78,12 @@ bool FInterchangeImportTest::RunTest(const FString& Path)
 
 	TSharedPtr<FInterchangeAutomationTestStepData, ESPMode::ThreadSafe> AutomationTestStepData = MakeShared<FInterchangeAutomationTestStepData, ESPMode::ThreadSafe>();
 	
-	FName PackageName = FName(FPaths::GetBaseFilename(Path, false));
-	FName PackagePath = FName(FPaths::GetPath(Path));
-	FName AssetName = FName(FPaths::GetBaseFilename(Path, true));
-	FTopLevelAssetPath ClassName = UInterchangeImportTestPlan::StaticClass()->GetClassPathName();
-	AutomationTestStepData->TestPlanData.AssetData = FAssetData(PackageName, PackagePath, AssetName, ClassName);
-	AutomationTestStepData->TestPlanData.TestPlan = CastChecked<UInterchangeImportTestPlan>(AutomationTestStepData->TestPlanData.AssetData.GetAsset());
+	const FName PackageName = FName(FPaths::GetBaseFilename(Path, false));
+	const FName PackagePath = FName(FPaths::GetPath(Path));
+	const FName AssetName = FName(FPaths::GetBaseFilename(Path, true));
+	const FTopLevelAssetPath ClassName = UInterchangeImportTestPlan::StaticClass()->GetClassPathName();
+	AutomationTestStepData->TestPlanData = FInterchangeImportTestData(PackageName, PackagePath, AssetName, ClassName);
+	AutomationTestStepData->TestPlanData.TestPlan = CastChecked<UInterchangeImportTestPlan>(AutomationTestStepData->TestPlanData.GetAsset());
 	
 	if (!ensure(AutomationTestStepData->TestPlanData.TestPlan))
 	{
@@ -91,9 +93,11 @@ bool FInterchangeImportTest::RunTest(const FString& Path)
 
 	// Base path to import assets to
 
-	FString SubDirToUse = TEXT("Interchange/ImportTest/");
-	FString BasePackagePath = FString(TEXT("/Temp/")) / SubDirToUse;
-	FString BaseFilePath = FPaths::ProjectSavedDir() / SubDirToUse;
+	const FString SubDirToUse = TEXT("Interchange/Temp/ImportTest/");
+	const FString BasePackagePath = FPaths::Combine(TEXT("/Game/Tests"), SubDirToUse);
+	const FString BaseFilePath = FPaths::Combine(FPaths::ProjectContentDir(), TEXT("Tests"), SubDirToUse);
+
+	AutomationTestStepData->PendingDeleteDirectoryPath = FPaths::Combine(FPaths::ProjectContentDir(), TEXT("Tests/Interchange/Temp/"));
 
 	// Clear out the folder contents before we do anything else
 
@@ -101,7 +105,7 @@ bool FInterchangeImportTest::RunTest(const FString& Path)
 	constexpr bool bDeleteRecursively = true;
 	IFileManager::Get().DeleteDirectory(*BaseFilePath, bRequireExists, bDeleteRecursively);
 
-	FString TestPlanAssetName = AutomationTestStepData->TestPlanData.AssetData.AssetName.ToString().Replace(TEXT("/"), TEXT("_"));
+	FString TestPlanAssetName = AutomationTestStepData->TestPlanData.AssetName.ToString().Replace(TEXT("/"), TEXT("_"));
 	AutomationTestStepData->TestPlanData.DestAssetPackagePath = BasePackagePath / TestPlanAssetName;
 	AutomationTestStepData->TestPlanData.DestAssetFilePath = BaseFilePath / TestPlanAssetName;
 	const bool bAddRecursively = true;
@@ -111,37 +115,73 @@ bool FInterchangeImportTest::RunTest(const FString& Path)
 	AutomationTestStepData->bSuccess = true;
 
 	const bool bRunSynchornously = AutomationTestStepData->TestPlanData.TestPlan->IsRunningSynchornously();
-	for (int32 TestStepIndex = 0; TestStepIndex < AutomationTestStepData->TestPlanData.TestPlan->Steps.Num(); ++TestStepIndex)
+
+	// If it is a level import, custom level or transient world is made available for the import.
+	AutomationTestStepData->TestPlanData.TestPlan->SetupLevelForImport();
+	
+	// Import
+	if (UInterchangeImportTestStepImport* TestStepImport = AutomationTestStepData->TestPlanData.TestPlan->ImportStep)
 	{
-		if (UInterchangeImportTestStepBase* TestStepBase = AutomationTestStepData->TestPlanData.TestPlan->Steps[TestStepIndex])
+		constexpr bool bIsReimportStep = false;
+		ADD_COMMAND(FInterchangeIntializeStepCommand(AutomationTestStepData.ToSharedRef(), bIsReimportStep, -1));
+		ADD_COMMAND(FInterchangeInterStepCollectResultCommand(AutomationTestStepData.ToSharedRef()));
+
+		// Screenshot Tests should only be performed via Test Automation Window
+		if (!bRunSynchornously && TestStepImport->HasScreenshotTest())
 		{
-			ADD_COMMAND(FInterchangeIntializeStepCommand(AutomationTestStepData.ToSharedRef(), TestStepIndex));
-			ADD_COMMAND(FInterchangeInterStepCollectResultCommand(AutomationTestStepData.ToSharedRef()));
+			// Wait so that all the materials applied to actors would be visible in the screenshot.
+			ADD_LATENT_AUTOMATION_COMMAND(FWaitForShadersToFinishCompilingInGame());
 
-			// Screenshot Tests should only be performed via Test Automation Window
-			if (!bRunSynchornously && TestStepBase->HasScreenshotTest())
+			const FInterchangeTestScreenshotParameters ScreenshotParameters = TestStepImport->GetScreenshotParameters();
+
+			// Focuses camera on relevant actors in the scene.
+			ADD_LATENT_AUTOMATION_COMMAND(FInterchangeSetupScreenshotViewportCommand(AutomationTestStepData.ToSharedRef(), ScreenshotParameters));
+
+			// Camera Viewport Transition Time.
+			ADD_LATENT_AUTOMATION_COMMAND(FEngineWaitLatentCommand(0.3f));
+
+			// Send a request for screenshot and wait until comparison is finished.
+			ADD_LATENT_AUTOMATION_COMMAND(FInterchangeCaptureScreenshotCommand(AutomationTestStepData.ToSharedRef(), ScreenshotParameters));
+		}
+
+		ADD_COMMAND(FInterchangeInterStepPerformTestsAndCollectGarbageCommand(AutomationTestStepData.ToSharedRef()));
+	}
+
+	// Reimport 
+	{
+		constexpr bool bIsReimportStep = true;
+		for (int32 TestStepIndex = 0; TestStepIndex < AutomationTestStepData->TestPlanData.TestPlan->ReimportStack.Num(); ++TestStepIndex)
+		{
+			if (UInterchangeImportTestStepReimport* ReimportTestStep = AutomationTestStepData->TestPlanData.TestPlan->ReimportStack[TestStepIndex])
 			{
-				// Wait so that all the materials applied to actors would be visible in the screenshot.
-				ADD_LATENT_AUTOMATION_COMMAND(FWaitForShadersToFinishCompilingInGame());
-				
-				const FInterchangeTestScreenshotParameters ScreenshotParameters = TestStepBase->GetScreenshotParameters();
-
-				// Focuses camera on relevant actors in the scene.
-				ADD_LATENT_AUTOMATION_COMMAND(FInterchangeSetupScreenshotViewportCommand(AutomationTestStepData.ToSharedRef(), ScreenshotParameters));
-				
-				// Camera Viewport Transition Time.
-				ADD_LATENT_AUTOMATION_COMMAND(FEngineWaitLatentCommand(0.3f));
-
-				// Send a request for screenshot and wait until comparison is finished.
-				ADD_LATENT_AUTOMATION_COMMAND(FInterchangeCaptureScreenshotCommand(AutomationTestStepData.ToSharedRef(), ScreenshotParameters));
+				ADD_COMMAND(FInterchangeIntializeStepCommand(AutomationTestStepData.ToSharedRef(), bIsReimportStep, TestStepIndex));
+				ADD_COMMAND(FInterchangeInterStepCollectResultCommand(AutomationTestStepData.ToSharedRef()));
+
+				// Screenshot Tests should only be performed via Test Automation Window
+				if (!bRunSynchornously && ReimportTestStep->HasScreenshotTest())
+				{
+					// Wait so that all the materials applied to actors would be visible in the screenshot.
+					ADD_LATENT_AUTOMATION_COMMAND(FWaitForShadersToFinishCompilingInGame());
+
+					const FInterchangeTestScreenshotParameters ScreenshotParameters = ReimportTestStep->GetScreenshotParameters();
+
+					// Focuses camera on relevant actors in the scene.
+					ADD_LATENT_AUTOMATION_COMMAND(FInterchangeSetupScreenshotViewportCommand(AutomationTestStepData.ToSharedRef(), ScreenshotParameters));
+
+					// Camera Viewport Transition Time.
+					ADD_LATENT_AUTOMATION_COMMAND(FEngineWaitLatentCommand(0.3f));
+
+					// Send a request for screenshot and wait until comparison is finished.
+					ADD_LATENT_AUTOMATION_COMMAND(FInterchangeCaptureScreenshotCommand(AutomationTestStepData.ToSharedRef(), ScreenshotParameters));
+				}
+
+				ADD_COMMAND(FInterchangeInterStepPerformTestsAndCollectGarbageCommand(AutomationTestStepData.ToSharedRef()));
 			}
-			
-			ADD_COMMAND(FInterchangeInterStepPerformTestsAndCollectGarbageCommand(AutomationTestStepData.ToSharedRef()));
 		}
 	}
 	ADD_COMMAND(FInterchangeTestCleanUpCommand(AutomationTestStepData.ToSharedRef()));
 	ADD_COMMAND(FInterchangeTestAutomationTestSuccessCommand(AutomationTestStepData.ToSharedRef()));
-	
+
 	return true;
 
 #endif
diff --git a/Engine/Plugins/Tests/InterchangeTests/Source/InterchangeTests/Private/InterchangeImportTestPlan.cpp b/Engine/Plugins/Tests/InterchangeTests/Source/InterchangeTests/Private/InterchangeImportTestPlan.cpp
--- a/Engine/Plugins/Tests/InterchangeTests/Source/InterchangeTests/Private/InterchangeImportTestPlan.cpp
+++ b/Engine/Plugins/Tests/InterchangeTests/Source/InterchangeTests/Private/InterchangeImportTestPlan.cpp
@@ -1,16 +1,47 @@
 // Copyright Epic Games, Inc. All Rights Reserved.
 
 #include "InterchangeImportTestPlan.h"
+
+#include "Editor.h"
+#include "Engine/Engine.h"
+#include "Engine/Level.h"
+#include "InterchangeImportTestStepImport.h"
+#include "InterchangeImportTestStepReimport.h"
 #include "JsonObjectConverter.h"
+#include "LevelEditorSubsystem.h"
+#include "Logging/MessageLog.h"
 #include "Misc/AutomationTest.h"
 #include "Misc/FileHelper.h"
 #include "Misc/Paths.h"
-#include "Logging/MessageLog.h"
 
 #include UE_INLINE_GENERATED_CPP_BY_NAME(InterchangeImportTestPlan)
 
 #define LOCTEXT_NAMESPACE "InterchangeImportTestPlan"
 
+#if WITH_EDITOR
+#include "HAL/IConsoleManager.h"
+
+static bool GInterchangeTestPlanPersistWarningDialogSuppresion = true;
+static FAutoConsoleVariableRef CCvarInterchangeDefaultShowEssentialsView(
+	TEXT("Interchange.TestPlan.PersistWarningDialogSuppression"),
+	GInterchangeTestPlanPersistWarningDialogSuppresion,
+	TEXT("Whether the suppresion choices for warning dialogs in interchange test plan persisted beyond the current session."),
+	ECVF_Default);
+
+
+FSuppressableWarningDialog::EMode UInterchangeImportTestPlan::GetInterchangeTestPlanWarningDialogMode()
+{
+	return GInterchangeTestPlanPersistWarningDialogSuppresion ? FSuppressableWarningDialog::EMode::Default : FSuppressableWarningDialog::EMode::DontPersistSuppressionAcrossSessions;
+}
+#endif
+
+
+UInterchangeImportTestPlan::UInterchangeImportTestPlan()
+{
+	ImportStep = CreateDefaultSubobject<UInterchangeImportTestStepImport>("ImportStep");
+	ImportStep->ParentTestPlan = this;
+	WorldPath = FSoftObjectPath(TEXT("/Game/Tests/Interchange/InterchangeTestMap.InterchangeTestMap"));
+}
 
 bool UInterchangeImportTestPlan::ReadFromJson(const FString& Filename)
 {
@@ -59,7 +90,7 @@ void UInterchangeImportTestPlan::RunThisTest()
 	AutomationEditorLog.Info(NewPageNameText);
 
 	FAutomationTestFramework& TestFramework = FAutomationTestFramework::Get();
-	
+
 	TestFramework.StartTestByName(FString(TEXT("FInterchangeImportTest ")) + GetPathName(), 0);
 
 	FAutomationTestExecutionInfo ExecutionInfo;
@@ -96,6 +127,206 @@ bool UInterchangeImportTestPlan::IsRunningSynchornously() const
 	return bRunSynchronously;
 }
 
+
+void UInterchangeImportTestPlan::PostLoad()
+{
+	Super::PostLoad();
+#if WITH_EDITOR
+	if (HasAnyFlags(RF_WasLoaded))
+	{
+		if (!Steps_DEPRECATED.IsEmpty())
+		{
+			bool bAssignedImportStep = false;
+			UInterchangeImportTestStepImport* FirstValidImportTestStep = nullptr;
+			// Because we are checking consecutive import and reimport steps.
+			for (int32 StepIndex = 0; StepIndex < Steps_DEPRECATED.Num(); ++StepIndex)
+			{
+				if (UInterchangeImportTestStepImport* ImportTestStep = Cast<UInterchangeImportTestStepImport>(Steps_DEPRECATED[StepIndex]))
+				{
+					if (!FirstValidImportTestStep)
+					{
+						FirstValidImportTestStep = ImportTestStep;
+					}
+					if (bAssignedImportStep)
+					{
+						break;
+					}
+				}
+
+				if (!bAssignedImportStep && (StepIndex < Steps_DEPRECATED.Num() - 1))
+				{
+					UInterchangeImportTestStepImport* ImportTestStep = Cast<UInterchangeImportTestStepImport>(Steps_DEPRECATED[StepIndex]);
+					UInterchangeImportTestStepReimport* ReimportTestStep = Cast<UInterchangeImportTestStepReimport>(Steps_DEPRECATED[StepIndex + 1]);
+					if (ImportTestStep && ReimportTestStep)
+					{
+						ImportStep = ImportTestStep;
+
+						bAssignedImportStep = true;
+					}
+				}
+
+				if (UInterchangeImportTestStepReimport* ReimportTestStep = Cast<UInterchangeImportTestStepReimport>(Steps_DEPRECATED[StepIndex]))
+				{
+					if (bAssignedImportStep)
+					{
+						ReimportStack.Add(ReimportTestStep);
+					}
+				}
+			}
+
+			// The Test Plan only contained Import Steps.
+			if (ReimportStack.IsEmpty() && FirstValidImportTestStep)
+			{
+				ImportStep = FirstValidImportTestStep;
+			}
+			Steps_DEPRECATED.Empty();
+			bChangeObjectOuters = true;
+		}
+		FCoreUObjectDelegates::OnAssetLoaded.AddUObject(this, &UInterchangeImportTestPlan::OnAssetLoaded);
+	}
+#endif
+}
+
+void UInterchangeImportTestPlan::OnAssetLoaded(UObject* Asset)
+{
+	if (Asset == this)
+	{
+		if (bChangeObjectOuters)
+		{
+			ImportStep->Rename(nullptr, this, REN_NonTransactional | REN_DontCreateRedirectors);
+		}
+		ImportStep->ParentTestPlan = this;
+
+		for (const TObjectPtr<UInterchangeImportTestStepReimport>& ReimportTestStep : ReimportStack)
+		{
+			if (bChangeObjectOuters)
+			{
+				ReimportTestStep->Rename(nullptr, this, REN_NonTransactional | REN_DontCreateRedirectors);
+			}
+			ReimportTestStep->InitializeReimportStepFromImportStep(ImportStep);
+			ReimportTestStep->ParentTestPlan = this;
+		}
+		
+		MarkPackageDirty();
+		// Unregister to the OnAssetLoad event as it is not needed anymore
+		FCoreUObjectDelegates::OnAssetLoaded.RemoveAll(this);
+	}
+}
+
+#if WITH_EDITOR
+
+void UInterchangeImportTestPlan::PreEditUndo()
+{
+	PrevReimportStackSize = ReimportStack.Num();
+}
+
+void UInterchangeImportTestPlan::PostEditUndo()
+{
+	if (PrevReimportStackSize != ReimportStack.Num())
+	{
+		for (int32 ReimportStepIndex = 0; ReimportStepIndex < ReimportStack.Num(); ++ReimportStepIndex)
+		{
+			if (ReimportStack[ReimportStepIndex] != nullptr)
+			{
+				ReimportStack[ReimportStepIndex]->RemoveImportStepPipelineSettingsModifiedDelegate();
+				ReimportStack[ReimportStepIndex]->InitializeReimportStepFromImportStep(ImportStep);
+			}
+		}
+		PrevReimportStackSize = ReimportStack.Num();
+	}
+}
+
+void UInterchangeImportTestPlan::PreEditChange(FProperty* PropertyAboutToChange)
+{
+	if (PropertyAboutToChange && (PropertyAboutToChange->GetFName() == GET_MEMBER_NAME_CHECKED(UInterchangeImportTestPlan, ReimportStack)))
+	{
+		for (int32 ReimportStepIndex = 0; ReimportStepIndex < ReimportStack.Num(); ++ReimportStepIndex)
+		{
+			if (ReimportStack[ReimportStepIndex] != nullptr)
+			{
+				ReimportStack[ReimportStepIndex]->RemoveImportStepPipelineSettingsModifiedDelegate();
+			}
+		}
+	}
+	Super::PreEditChange(PropertyAboutToChange);
+}
+
+void UInterchangeImportTestPlan::PostEditChangeChainProperty(FPropertyChangedChainEvent& PropertyChangedEvent)
+{
+	if (PropertyChangedEvent.GetPropertyName() == GET_MEMBER_NAME_CHECKED(UInterchangeImportTestPlan, ReimportStack))
+	{
+		for (int32 ReimportStepIndex = 0; ReimportStepIndex < ReimportStack.Num(); ++ReimportStepIndex)
+		{
+			if (ReimportStack[ReimportStepIndex] != nullptr && !ImportStep->OnImportTestStepDataChanged.IsBoundToObject(ReimportStack[ReimportStepIndex]))
+			{
+				ReimportStack[ReimportStepIndex]->InitializeReimportStepFromImportStep(ImportStep);
+			}
+		}
+	}
+	Super::PostEditChangeChainProperty(PropertyChangedEvent);
+}
+#endif
+
+void UInterchangeImportTestPlan::SetupLevelForImport()
+{
+	if (!ImportStep->bImportIntoLevel)
+	{
+		return;
+	}
+
+	// Create transient world to host data from producer
+	if (!IsRunningSynchornously() && WorldPath.IsValid())
+	{
+		if (ULevelEditorSubsystem* LevelEditorSubsystem = GEditor->GetEditorSubsystem<ULevelEditorSubsystem>())
+		{
+			LevelEditorSubsystem->LoadLevel(WorldPath.GetAssetPathString());
+		}
+	}
+	else
+	{
+		TransientWorld = TStrongObjectPtr<UWorld>(NewObject<UWorld>(GetTransientPackage()));
+		TransientWorld->WorldType = EWorldType::EditorPreview;
+
+		FWorldContext& WorldContext = GEngine->CreateNewWorldContext(TransientWorld->WorldType);
+		WorldContext.SetCurrentWorld(TransientWorld.Get());
+
+		TransientWorld->InitializeNewWorld(UWorld::InitializationValues()
+			.AllowAudioPlayback(false)
+			.CreatePhysicsScene(false)
+			.RequiresHitProxies(false)
+			.CreateNavigation(false)
+			.CreateAISystem(false)
+			.ShouldSimulatePhysics(false)
+			.SetTransactional(false)
+		);
+
+		UE_LOG(LogTemp, Display, TEXT("Test Plan World Path Name: %s"), *(TransientWorld.Get()->GetPathName()));
+	}
+}
+
+void UInterchangeImportTestPlan::CleanupLevel()
+{
+	if (!ImportStep->bImportIntoLevel)
+	{
+		return;
+	}
+
+	if (!IsRunningSynchornously() && WorldPath.IsValid())
+	{
+		if (ULevelEditorSubsystem* LevelEditorSubsystem = GEditor->GetEditorSubsystem<ULevelEditorSubsystem>())
+		{
+			LevelEditorSubsystem->LoadLevel(WorldPath.GetAssetPathString());
+		}
+	}
+	else if (TransientWorld)
+	{
+		// Now delete world
+		GEngine->DestroyWorldContext(TransientWorld.Get());
+		TransientWorld->DestroyWorld(true);
+		TransientWorld.Reset();
+	}
+}
+
 namespace UE::Interchange
 {
 	FString FInterchangeImportTestPlanStaticHelpers::GetTestNameFromObjectPathString(const FString& InObjectPathString, bool bAddBeautifiedTestNamePrefix /*= false*/)
@@ -107,10 +338,10 @@ namespace UE::Interchange
 		FPaths::MakePathRelativeTo(BeautifiedName, *GetInterchangeImportTestRootGameFolder());
 		BeautifiedName.ReplaceCharInline(TEXT('/'), TEXT('.'));
 		BeautifiedName.ReplaceCharInline(TEXT('\\'), TEXT('.'));
-		
+
 		if (bAddBeautifiedTestNamePrefix)
 		{
-			return FString::Printf(TEXT("%s.%s"), *GetBeautifiedTestName() , *BeautifiedName);
+			return FString::Printf(TEXT("%s.%s"), *GetBeautifiedTestName(), *BeautifiedName);
 		}
 
 		return BeautifiedName;
@@ -125,6 +356,5 @@ namespace UE::Interchange
 	{
 		return TEXT("/Game/Tests/Interchange/");
 	}
-
 }
-#undef LOCTEXT_NAMESPACE
+#undef LOCTEXT_NAMESPACE
\ No newline at end of file
diff --git a/Engine/Plugins/Tests/InterchangeTests/Source/InterchangeTests/Private/InterchangeImportTestStepImport.cpp b/Engine/Plugins/Tests/InterchangeTests/Source/InterchangeTests/Private/InterchangeImportTestStepImport.cpp
--- a/Engine/Plugins/Tests/InterchangeTests/Source/InterchangeTests/Private/InterchangeImportTestStepImport.cpp
+++ b/Engine/Plugins/Tests/InterchangeTests/Source/InterchangeTests/Private/InterchangeImportTestStepImport.cpp
@@ -3,10 +3,6 @@
 #include "InterchangeImportTestStepImport.h"
 
 #include "Dialogs/Dialogs.h"
-#include "Editor.h"
-#include "Engine/Engine.h"
-#include "Engine/Level.h"
-#include "Engine/World.h"
 #include "GameFramework/WorldSettings.h"
 #include "HAL/FileManager.h"
 #include "InterchangeImportTestData.h"
@@ -14,7 +10,7 @@
 #include "InterchangeManager.h"
 #include "InterchangeProjectSettings.h"
 #include "InterchangeTestsLog.h"
-#include "LevelEditorSubsystem.h"
+#include "InterchangeTranslatorBase.h"
 #include "Misc/PackageName.h"
 #include "Misc/ScopedSlowTask.h"
 #include "ScopedTransaction.h"
@@ -23,7 +19,7 @@
 
 #define LOCTEXT_NAMESPACE "InterchangeImportTestStepImport"
 
-UInterchangeImportTestStepImport::UInterchangeImportTestStepImport(const FObjectInitializer& InObjectInitializer)
+UInterchangeImportTestStepImport::UInterchangeImportTestStepImport()
 {
 	PipelineSettings.ParentTestStep = this;
 }
@@ -96,36 +92,11 @@ TTuple<UE::Interchange::FAssetImportResultPtr, UE::Interchange::FSceneImportResu
 
 	if (bImportIntoLevel)
 	{
-		// Create transient world to host data from producer
-		if (!Data.TestPlan->IsRunningSynchornously() && WorldPath.IsValid())
+		// Use the world from the Test Plan for Level Import 
+		if (ensure(Data.TestPlan))
 		{
-			if (ULevelEditorSubsystem* LevelEditorSubsystem = GEditor->GetEditorSubsystem<ULevelEditorSubsystem>())
-			{
-				LevelEditorSubsystem->LoadLevel(WorldPath.GetAssetPathString());
-				Params.ImportLevel = GWorld->GetCurrentLevel();
-			}
-		}
-		else 
-		{
-			TransientWorld = TStrongObjectPtr<UWorld>(NewObject<UWorld>(GetTransientPackage()));
-			TransientWorld->WorldType = EWorldType::EditorPreview;
-
-			FWorldContext& WorldContext = GEngine->CreateNewWorldContext(TransientWorld->WorldType);
-			WorldContext.SetCurrentWorld(TransientWorld.Get());
-
-			TransientWorld->InitializeNewWorld(UWorld::InitializationValues()
-				.AllowAudioPlayback(false)
-				.CreatePhysicsScene(false)
-				.RequiresHitProxies(false)
-				.CreateNavigation(false)
-				.CreateAISystem(false)
-				.ShouldSimulatePhysics(false)
-				.SetTransactional(false)
-			);
-
-			Params.ImportLevel = TransientWorld->GetCurrentLevel();
+			Params.ImportLevel = Data.TestPlan->GetCurrentLevel();
 		}
-
 		return InterchangeManager.ImportSceneAsync(Data.DestAssetPackagePath, ScopedSourceData.GetSourceData(), Params);
 	}
 	else
@@ -139,13 +110,6 @@ FTestStepResults UInterchangeImportTestStepImport::FinishStep(FInterchangeImport
 {
 	FTestStepResults Results;
 
-	// If we need to perform a save and fresh reload of everything we imported, do it here
-	// Incompatible with bImportIntoLevel for now, since we'd need to patch all references to the assets
-	if (bSaveThenReloadImportedAssets && !bImportIntoLevel)
-	{
-		SaveReloadAssets(Data);
-	}
-
 	// Run all the tests
 	bool bSuccess = PerformTests(Data, CurrentTest);
 	
@@ -155,30 +119,12 @@ FTestStepResults UInterchangeImportTestStepImport::FinishStep(FInterchangeImport
 
 	for (UObject* Object : ResultObjects)
 	{
-		if (!Object->IsA<AActor>() && !Object->IsA<UWorld>())
+		if (!Object->IsA<AActor>())
 		{
 			Data.ResultObjects.Add(Object);
 		}
 	}
 
-	if (bImportIntoLevel)
-	{
-		if (!Data.TestPlan->IsRunningSynchornously() && WorldPath.IsValid())
-		{
-			if (ULevelEditorSubsystem* LevelEditorSubsystem = GEditor->GetEditorSubsystem<ULevelEditorSubsystem>())
-			{
-				LevelEditorSubsystem->LoadLevel(WorldPath.GetAssetPathString());
-			}
-		}
-		else if (TransientWorld)
-		{
-			// Now delete world
-			GEngine->DestroyWorldContext(TransientWorld.Get());
-			TransientWorld->DestroyWorld(true);
-			TransientWorld.Reset();
-		}
-	}
-
 	Results.bTestStepSuccess = bSuccess;
 	return Results;
 }
@@ -190,7 +136,7 @@ FString UInterchangeImportTestStepImport::GetContextString() const
 
 bool UInterchangeImportTestStepImport::HasScreenshotTest() const
 {
-	return bImportIntoLevel && WorldPath.IsValid();
+	return bImportIntoLevel && bTakeScreenshot;
 }
 
 FInterchangeTestScreenshotParameters UInterchangeImportTestStepImport::GetScreenshotParameters() const
@@ -379,9 +325,8 @@ void UInterchangeImportTestStepImport::EditPipelineSettings()
 			PipelineStack.Empty(OutPipelines.Num());
 			for (UInterchangePipelineBase* Pipeline : OutPipelines)
 			{
-				if (UInterchangePipelineBase* InstancedPipeline = DuplicateObject<UInterchangePipelineBase>(Pipeline, GetOuter()))
+				if (UInterchangePipelineBase* InstancedPipeline = DuplicateObject<UInterchangePipelineBase>(Pipeline, this))
 				{
-					InstancedPipeline->SetFlags(EObjectFlags::RF_Transactional);
 					PipelineStack.Add(InstancedPipeline);
 				}
 			}
@@ -390,6 +335,8 @@ void UInterchangeImportTestStepImport::EditPipelineSettings()
 		{
 			PipelineSettings.UpdatePipelines(OutPipelines);
 		}
+
+		BroadcastImportStepChangedEvent(EImportStepDataChangeType::PipelineSettings);
 	}
 #endif 
 }
@@ -407,19 +354,56 @@ void UInterchangeImportTestStepImport::ClearPipelineSettings()
 	{
 		PipelineSettings.ClearPipelines();
 	}
+	
+	BroadcastImportStepChangedEvent(EImportStepDataChangeType::PipelineSettings);
 #endif
 }
 
+bool UInterchangeImportTestStepImport::IsUsingOverridePipelines(bool bCheckForValidPipelines) const
+{
+	if (!bCheckForValidPipelines)
+	{
+		return bUseOverridePipelineStack;
+	}
+
+	if (bUseOverridePipelineStack)
+	{
+		for (const TObjectPtr<UInterchangePipelineBase>& Pipeline : PipelineStack)
+		{
+			if (!Pipeline)
+			{
+				return false;
+			}
+		}
+		return true;
+	}
+
+	return false;
+}
+
 bool UInterchangeImportTestStepImport::ShouldImportIntoLevelChangeRequireMessageBox() const
 {
 	return !bUseOverridePipelineStack && PipelineSettings.IsUsingModifiedSettings();
 }
 
+void UInterchangeImportTestStepImport::PostLoad()
+{
+	Super::PostLoad();
+	if (!SourceFile.FilePath.IsEmpty() && LastSourceFileExtension.IsEmpty())
+	{
+		LastSourceFileExtension = FPaths::GetExtension(SourceFile.FilePath);
+	}
+}
+
 #if WITH_EDITOR
 void UInterchangeImportTestStepImport::PostEditChangeProperty(FPropertyChangedEvent& PropertyChangedEvent)
 {
-	if (PropertyChangedEvent.Property->GetFName() == GET_MEMBER_NAME_CHECKED(UInterchangeImportTestStepImport, bUseOverridePipelineStack))
+	static const FName UseOverridePipelineStackPropertyName = GET_MEMBER_NAME_CHECKED(UInterchangeImportTestStepImport, bUseOverridePipelineStack);
+	static const FName ImportIntoLevelPropertyName = GET_MEMBER_NAME_CHECKED(UInterchangeImportTestStepImport, bImportIntoLevel);
+
+	if (PropertyChangedEvent.Property->GetFName() == UseOverridePipelineStackPropertyName)
 	{
+		bool bSendPipelineSettingsChangedEvent = false;
 		if (bUseOverridePipelineStack && !PipelineSettings.CustomPipelines.IsEmpty())
 		{
 			FSuppressableWarningDialog::FSetupInfo DialogSetupInfo = FSuppressableWarningDialog::FSetupInfo(
@@ -430,14 +414,26 @@ void UInterchangeImportTestStepImport::PostEditChangeProperty(FPropertyChangedEv
 			DialogSetupInfo.ConfirmText = LOCTEXT("UseOverridePipelinesDialogOptionConfirm", "OK");
 			DialogSetupInfo.bDefaultToSuppressInTheFuture = false;
 			DialogSetupInfo.CheckBoxText = LOCTEXT("UseOverridePipelinesDialogCheckBoxText", "Don't show this dialog again");
+			DialogSetupInfo.DialogMode = UInterchangeImportTestPlan::GetInterchangeTestPlanWarningDialogMode();
 			FSuppressableWarningDialog SuppressableWarningDialog(DialogSetupInfo);
 			SuppressableWarningDialog.ShowModal();
 
 			PipelineSettings.ClearPipelines();
+			bSendPipelineSettingsChangedEvent = true;
+		}
+		
+		if (PipelineStack.Num() > 0)
+		{
+			bSendPipelineSettingsChangedEvent = true;
 		}
-	}
 
-	if (PropertyChangedEvent.Property->GetFName() == GET_MEMBER_NAME_CHECKED(UInterchangeImportTestStepImport, bImportIntoLevel))
+		if (bSendPipelineSettingsChangedEvent)
+		{
+			BroadcastImportStepChangedEvent(EImportStepDataChangeType::PipelineSettings);
+		}
+	}
+	
+	if (PropertyChangedEvent.Property->GetFName() == ImportIntoLevelPropertyName)
 	{
 		if (ShouldImportIntoLevelChangeRequireMessageBox())
 		{
@@ -450,18 +446,236 @@ void UInterchangeImportTestStepImport::PostEditChangeProperty(FPropertyChangedEv
 			DialogSetupInfo.CancelText =  LOCTEXT("ImportIntoLevelDialogOptionCancel", "No");
 			DialogSetupInfo.bDefaultToSuppressInTheFuture = false;
 			DialogSetupInfo.CheckBoxText = LOCTEXT("ImportIntoLevelDialogCheckBoxText", "Don't show this dialog again");
+			DialogSetupInfo.DialogMode = UInterchangeImportTestPlan::GetInterchangeTestPlanWarningDialogMode();
 			FSuppressableWarningDialog SuppressableWarningDialog(DialogSetupInfo);
 			FSuppressableWarningDialog::EResult DialogResult = SuppressableWarningDialog.ShowModal();
 			if (DialogResult != FSuppressableWarningDialog::Cancel)
 			{
 				PipelineSettings.ClearPipelines();
+				BroadcastImportStepChangedEvent(EImportStepDataChangeType::ImportIntoLevel);
 			}
 			else
 			{
 				bImportIntoLevel = !bImportIntoLevel;
 			}
 		}
+		else
+		{
+			BroadcastImportStepChangedEvent(EImportStepDataChangeType::ImportIntoLevel);
+		}
+	}
+
+	Super::PostEditChangeProperty(PropertyChangedEvent);
+}
+
+void UInterchangeImportTestStepImport::PostEditChangeChainProperty(FPropertyChangedChainEvent& PropertyChangedEvent)
+{
+	Super::PostEditChangeChainProperty(PropertyChangedEvent);
+
+	static const FName PipelineStackPropertyName = GET_MEMBER_NAME_CHECKED(UInterchangeImportTestStepImport, PipelineStack);
+	static const FName SourceFilePropertyName = GET_MEMBER_NAME_CHECKED(UInterchangeImportTestStepImport, SourceFile);
+
+	if (PropertyChangedEvent.GetPropertyName() == PipelineStackPropertyName)
+	{
+		if (PropertyChangedEvent.ChangeType == EPropertyChangeType::ValueSet 
+			|| PropertyChangedEvent.ChangeType == EPropertyChangeType::ArrayClear 
+			|| PropertyChangedEvent.ChangeType == EPropertyChangeType::ArrayRemove)
+		{
+			BroadcastImportStepChangedEvent(EImportStepDataChangeType::PipelineSettings);
+		}
+	}
+
+	const FProperty* ActiveMemberPropertyNode = PropertyChangedEvent.PropertyChain.GetActiveMemberNode()->GetValue();
+	if (ActiveMemberPropertyNode && (ActiveMemberPropertyNode->GetFName() == SourceFilePropertyName))
+	{
+		const FString CurrSourceFileExtension = FPaths::GetExtension(SourceFile.FilePath);
+		if (!LastSourceFileExtension.IsEmpty() && !CurrSourceFileExtension.IsEmpty()
+			&& LastSourceFileExtension != CurrSourceFileExtension
+			&& !bUseOverridePipelineStack
+			&& !PipelineSettings.CustomPipelines.IsEmpty())
+		{
+			FSuppressableWarningDialog::FSetupInfo DialogSetupInfo = FSuppressableWarningDialog::FSetupInfo(
+				LOCTEXT("SourceFileExtensionChangedDialogText", "Current pipelines might not be compatible with the new source file. Clearing out to restore the defaults."),
+				LOCTEXT("SourceFileExtensionChangedDialogTitle", "Source File Extension Changed"),
+				"InterchangeImportTestPlanSourceFileExtensionChangedWarning");
+
+			DialogSetupInfo.ConfirmText = LOCTEXT("SourceFileExtensionChangedDialogOptionConfirm", "OK");
+			DialogSetupInfo.bDefaultToSuppressInTheFuture = false;
+			DialogSetupInfo.CheckBoxText = LOCTEXT("SourceFileExtensionChangedDialogCheckBoxText", "Don't show this dialog again");
+			DialogSetupInfo.DialogMode = UInterchangeImportTestPlan::GetInterchangeTestPlanWarningDialogMode();
+			FSuppressableWarningDialog SuppressableWarningDialog(DialogSetupInfo);
+			SuppressableWarningDialog.ShowModal();
+
+			PipelineSettings.ClearPipelines();
+		}
+		LastSourceFileExtension = CurrSourceFileExtension;
+
+		BroadcastImportStepChangedEvent(EImportStepDataChangeType::SourceFile);
+	}
+}
+
+TArray<TObjectPtr<UInterchangePipelineBase>> UInterchangeImportTestStepImport::GetCurrentPipelinesOrDefault() const
+{
+	TArray<TObjectPtr<UInterchangePipelineBase>> OutPipelines;
+	if (bUseOverridePipelineStack)
+	{
+		OutPipelines = PipelineStack;
+	}
+	else if (!PipelineSettings.CustomPipelines.IsEmpty())
+	{
+		OutPipelines = PipelineSettings.CustomPipelines;
+	}
+	else
+	{
+		if (SourceFile.FilePath.IsEmpty())
+		{
+			return OutPipelines;
+		}
+
+		UE::Interchange::FScopedSourceData ScopedSourceData(SourceFile.FilePath);
+		UE::Interchange::FScopedTranslator ScopedTranslator(ScopedSourceData.GetSourceData());
+		if (!ScopedTranslator.GetTranslator())
+		{
+			UE_LOG(LogInterchangeTests, Error, TEXT("Cannot import file. The source data is not supported. Try enabling the [%s] extension for Interchange."), *FPaths::GetExtension(ScopedSourceData.GetSourceData()->GetFilename()));
+			return OutPipelines;
+		}
+
+		const FInterchangeImportSettings& InterchangeImportSettings = FInterchangeProjectSettingsUtils::GetDefaultImportSettings(bImportIntoLevel);
+		if (InterchangeImportSettings.PipelineStacks.Num() == 0)
+		{
+			UE_LOG(LogInterchangeTests, Error, TEXT("Failed to configure pipelines. There is no pipeline stack defined for the content import type."));
+			return OutPipelines;
+		}
+
+		if (!InterchangeImportSettings.PipelineStacks.Contains(InterchangeImportSettings.DefaultPipelineStack))
+		{
+			FInterchangeImportSettings& MutableInterchangeImportSettings = FInterchangeProjectSettingsUtils::GetMutableDefaultImportSettings(bImportIntoLevel);
+			TArray<FName> Keys;
+			MutableInterchangeImportSettings.PipelineStacks.GetKeys(Keys);
+			check(Keys.Num() > 0);
+			MutableInterchangeImportSettings.DefaultPipelineStack = Keys[0];
+		}
+
+		UE::Interchange::FScopedBaseNodeContainer ScopedBaseNodeContainer;
+
+		UInterchangePipelineConfigurationBase* RegisteredPipelineConfiguration = nullptr;
+		TSoftClassPtr <UInterchangePipelineConfigurationBase> ImportDialogClass = InterchangeImportSettings.ImportDialogClass;
+		if (ImportDialogClass.IsValid())
+		{
+			UClass* PipelineConfigurationClass = ImportDialogClass.LoadSynchronous();
+			if (PipelineConfigurationClass)
+			{
+				RegisteredPipelineConfiguration = NewObject<UInterchangePipelineConfigurationBase>(GetTransientPackage(), PipelineConfigurationClass, NAME_None, RF_NoFlags);
+				if (RegisteredPipelineConfiguration == nullptr)
+				{
+					UE_LOG(LogInterchangeTests, Error, TEXT("Failed to create a pipeline configuration object."));
+					return OutPipelines;
+				}
+			}
+		}
+
+		auto AdjustPipelineSettingForContext = [this, &ScopedBaseNodeContainer](UInterchangePipelineBase* Pipeline)
+			{
+				EInterchangePipelineContext Context;
+				Context = bImportIntoLevel ? EInterchangePipelineContext::SceneImport : EInterchangePipelineContext::AssetImport;
+				FInterchangePipelineContextParams ContextParams;
+				ContextParams.ContextType = Context;
+				ContextParams.BaseNodeContainer = ScopedBaseNodeContainer.GetBaseNodeContainer();
+				Pipeline->SetFromReimportOrOverride(bUseOverridePipelineStack);
+				Pipeline->AdjustSettingsForContext(ContextParams);
+			};
+
+		const bool bCanTranslate = UInterchangeManager::GetInterchangeManager().CanUseTranslator(ScopedTranslator.GetTranslator());
+		if (bCanTranslate)
+		{
+			FScopedSlowTask Progress(2.f, LOCTEXT("TranslatingSourceFile...", "Translating source file..."));
+			Progress.MakeDialog();
+			Progress.EnterProgressFrame(1.f);
+
+			UInterchangeBaseNodeContainer& BaseNodeContainer = *(ScopedBaseNodeContainer.GetBaseNodeContainer());
+			ScopedTranslator.GetTranslator()->SetResultsContainer(NewObject<UInterchangeResultsContainer>(GetTransientPackage(), NAME_None));
+			ScopedTranslator.GetTranslator()->Translate(BaseNodeContainer);
+
+			Progress.EnterProgressFrame(1.f);
+		}
+
+		const TMap<FName, FInterchangePipelineStack>& DefaultPipelineStacks = InterchangeImportSettings.PipelineStacks;
+		TArray<FInterchangeStackInfo> InPipelineStacks;
+
+		for (const TPair<FName, FInterchangePipelineStack>& PipelineStackInfo : DefaultPipelineStacks)
+		{
+			FName StackName = PipelineStackInfo.Key;
+			FInterchangeStackInfo& StackInfo = InPipelineStacks.AddDefaulted_GetRef();
+			StackInfo.StackName = StackName;
+
+			const FInterchangePipelineStack& DefPipelineStack = PipelineStackInfo.Value;
+			const TArray<FSoftObjectPath>* Pipelines = &DefPipelineStack.Pipelines;
+
+			// If applicable, check to see if a specific pipeline stack is associated with this translator
+			for (const FInterchangeTranslatorPipelines& TranslatorPipelines : DefPipelineStack.PerTranslatorPipelines)
+			{
+				const UClass* TranslatorClass = TranslatorPipelines.Translator.LoadSynchronous();
+				if (ScopedTranslator.GetTranslator() && ScopedTranslator.GetTranslator()->IsA(TranslatorClass))
+				{
+					Pipelines = &TranslatorPipelines.Pipelines;
+					break;
+				}
+			}
+
+			for (int32 PipelineIndex = 0; PipelineIndex < Pipelines->Num(); ++PipelineIndex)
+			{
+				if (UInterchangePipelineBase* GeneratedPipeline = UE::Interchange::GeneratePipelineInstance((*Pipelines)[PipelineIndex]))
+				{
+					AdjustPipelineSettingForContext(GeneratedPipeline);
+					StackInfo.Pipelines.Add(GeneratedPipeline);
+				}
+			}
+		}
+
+		FName DefaultStackName = FInterchangeProjectSettingsUtils::GetDefaultPipelineStackName(bImportIntoLevel, *ScopedSourceData.GetSourceData());
+		FInterchangeStackInfo* StackInfoPtr = InPipelineStacks.FindByPredicate([DefaultStackName](const FInterchangeStackInfo& StackInfo)
+			{
+				return StackInfo.StackName == DefaultStackName;
+			});
+
+		if (StackInfoPtr)
+		{
+			//When we do not show the UI we use the original stack
+			OutPipelines = StackInfoPtr->Pipelines;
+		}
+		else if (InPipelineStacks.Num() > 0)
+		{
+			//Take the first valid stack
+			for (FInterchangeStackInfo& StackInfo : InPipelineStacks)
+			{
+				if (!StackInfo.Pipelines.IsEmpty())
+				{
+					OutPipelines = StackInfo.Pipelines;
+					break;
+				}
+			}
+		}
+		else
+		{
+			UE_LOG(LogInterchangeTests, Error, TEXT("Interchange Test Plan: Cannot find any valid stack. Could not build shared test plan data."));
+		}
 	}
+
+	return OutPipelines;
+}
+
+void UInterchangeImportTestStepImport::BroadcastImportStepChangedEvent(EImportStepDataChangeType ChangeType)
+{
+	if (ChangeType == EImportStepDataChangeType::Unknown)
+	{
+		return;
+	}
+
+	FImportStepChangedData ChangedData;
+	ChangedData.ChangeType = ChangeType;
+	ChangedData.ImportStep = this;
+	OnImportTestStepDataChanged.Broadcast(ChangedData);
+
 }
 #endif
 
diff --git a/Engine/Plugins/Tests/InterchangeTests/Source/InterchangeTests/Private/InterchangeImportTestStepReimport.cpp b/Engine/Plugins/Tests/InterchangeTests/Source/InterchangeTests/Private/InterchangeImportTestStepReimport.cpp
--- a/Engine/Plugins/Tests/InterchangeTests/Source/InterchangeTests/Private/InterchangeImportTestStepReimport.cpp
+++ b/Engine/Plugins/Tests/InterchangeTests/Source/InterchangeTests/Private/InterchangeImportTestStepReimport.cpp
@@ -2,11 +2,52 @@
 
 #include "InterchangeImportTestStepReimport.h"
 
+#include "Dialogs/Dialogs.h"
+#include "GameFramework/Actor.h"
 #include "InterchangeImportTestData.h"
+#include "InterchangeImportTestPlan.h"
+#include "InterchangeImportTestStepImport.h"
+#include "InterchangeManager.h"
+#include "InterchangePipelineBase.h"
+#include "InterchangeProjectSettings.h"
+#include "InterchangeSceneImportAsset.h"
+#include "InterchangeTestsLog.h"
+#include "JsonObjectConverter.h"
+#include "Misc/AutomationTest.h"
 #include "Misc/PackageName.h"
+#include "Misc/ScopedSlowTask.h"
+#include "ScopedTransaction.h"
 
 #include UE_INLINE_GENERATED_CPP_BY_NAME(InterchangeImportTestStepReimport)
 
+#define LOCTEXT_NAMESPACE "InterchangeImportTestStepReimport"
+
+UInterchangeImportTestStepReimport::UInterchangeImportTestStepReimport()
+{
+	PipelineSettings.ParentTestStep = this;
+}
+
+void UInterchangeImportTestStepReimport::InitializeReimportStepFromImportStep(UInterchangeImportTestStepImport* ImportTestStep)
+{
+	if (ImportTestStep)
+	{
+		CachedImportStep = ImportTestStep;
+		ImportTestStep->OnImportTestStepDataChanged.AddUObject(this, &UInterchangeImportTestStepReimport::HandleImportPipelineSettingsModified);
+		if (LastSourceFileExtension.IsEmpty() && !ImportTestStep->SourceFile.FilePath.IsEmpty())
+		{
+			LastSourceFileExtension = FPaths::GetExtension(ImportTestStep->SourceFile.FilePath);
+		}
+	}
+}
+
+void UInterchangeImportTestStepReimport::RemoveImportStepPipelineSettingsModifiedDelegate()
+{
+	if (UInterchangeImportTestStepImport* ImportStep = CachedImportStep.Get())
+	{
+		ImportStep->OnImportTestStepDataChanged.RemoveAll(this);
+	}
+}
+
 TTuple<UE::Interchange::FAssetImportResultPtr, UE::Interchange::FSceneImportResultPtr> UInterchangeImportTestStepReimport::StartStep(FInterchangeImportTestData& Data)
 {
 	// Find the object we wish to reimport
@@ -38,14 +79,32 @@ TTuple<UE::Interchange::FAssetImportResultPtr, UE::Interchange::FSceneImportResu
 			}
 		}
 	}
-
+	
+	const FString SourceFilePathString = GetReimportStepSourceFilePathString();
 	if (AssetToReimport == nullptr)
 	{
-		return {nullptr, nullptr};
+		if (SourceFilePathString.IsEmpty())
+		{
+			return {nullptr, nullptr};
+		}
+		else
+		{
+			if (FAutomationTestBase* CurrentTest = FAutomationTestFramework::Get().GetCurrentTest())
+			{
+				CurrentTest->AddInfo(TEXT("Could not find any asset to re-import. Performing an import into the same directory with new file(might get converted to re-import if it has assets with same name)."));
+			}
+		}
 	}
 
+	const bool bIsSceneImport = (bImportIntoLevel 
+									&& !SourceFilePathString.IsEmpty()) 
+								|| (CachedImportStep.IsValid()
+									&& CachedImportStep->bImportIntoLevel
+									&& (AssetTypeToReimport.Get() == UInterchangeSceneImportAsset::StaticClass()));
+
+	
 	// Start the Interchange import
-	UE::Interchange::FScopedSourceData ScopedSourceData(SourceFileToReimport.FilePath);
+	UE::Interchange::FScopedSourceData ScopedSourceData(SourceFilePathString);
 
 	FImportAssetParameters Params;
 	if (bUseOverridePipelineStack)
@@ -56,12 +115,33 @@ TTuple<UE::Interchange::FAssetImportResultPtr, UE::Interchange::FSceneImportResu
 			Params.OverridePipelines.Add(Pipeline);
 		}
 	}
+	else if (PipelineSettings.CustomPipelines.Num() > 0)
+	{
+		Params.OverridePipelines.Reserve(PipelineSettings.CustomPipelines.Num());
+		for (TObjectPtr<UInterchangePipelineBase> Pipeline : PipelineSettings.CustomPipelines)
+		{
+			Params.OverridePipelines.Add(Pipeline);
+		}
+	}
 	Params.bIsAutomated = true;
+	Params.ImportLevel = CachedImportStep.IsValid() && CachedImportStep->bImportIntoLevel ? Data.TestPlan->GetCurrentLevel() : nullptr;
 	Params.ReimportAsset = AssetToReimport;
-
+	
 	UInterchangeManager& InterchangeManager = UInterchangeManager::GetInterchangeManager();
-	UE::Interchange::FAssetImportResultPtr AssetImportResult = InterchangeManager.ImportAssetAsync(Data.DestAssetPackagePath, ScopedSourceData.GetSourceData(), Params);
-	return {AssetImportResult, nullptr};
+	if (bIsSceneImport)
+	{
+		if (UInterchangeSceneImportAsset* SceneImportAsset = Cast<UInterchangeSceneImportAsset>(AssetToReimport))
+		{
+			Params.ImportLevel = nullptr;
+		}
+
+		return InterchangeManager.ImportSceneAsync(Data.DestAssetPackagePath, ScopedSourceData.GetSourceData(), Params);
+	}
+	else
+	{
+		UE::Interchange::FAssetImportResultPtr AssetImportResult = InterchangeManager.ImportAssetAsync(Data.DestAssetPackagePath, ScopedSourceData.GetSourceData(), Params);
+		return {AssetImportResult, nullptr};
+	}
 }
 
 
@@ -69,20 +149,25 @@ FTestStepResults UInterchangeImportTestStepReimport::FinishStep(FInterchangeImpo
 {
 	FTestStepResults Results;
 
-	// If we need to perform a save and fresh reload of everything we imported, do it here
-	if (bSaveThenReloadImportedAssets)
-	{
-		SaveReloadAssets(Data);
-	}
-
 	// Run all the tests
 	bool bSuccess = PerformTests(Data, CurrentTest);
 
+	// Only keep assets as result objects since the world and its actors are being destroyed
+	TArray<UObject*> ResultObjects = MoveTemp(Data.ResultObjects);
+	Data.ResultObjects.Reset(ResultObjects.Num());
+
+	for (UObject* Object : ResultObjects)
+	{
+		if (!Object->IsA<AActor>())
+		{
+			Data.ResultObjects.Add(Object);
+		}
+	}
+
 	Results.bTestStepSuccess = bSuccess;
 	return Results;
 }
 
-
 FString UInterchangeImportTestStepReimport::GetContextString() const
 {
 	return FString(TEXT("Reimporting ")) + FPaths::GetCleanFilename(SourceFileToReimport.FilePath);
@@ -90,23 +175,384 @@ FString UInterchangeImportTestStepReimport::GetContextString() const
 
 bool UInterchangeImportTestStepReimport::HasScreenshotTest() const
 {
-	return false;
+	return CachedImportStep.IsValid() && CachedImportStep->bImportIntoLevel && bTakeScreenshot;
 }
 
 FInterchangeTestScreenshotParameters UInterchangeImportTestStepReimport::GetScreenshotParameters() const
 {
-	return FInterchangeTestScreenshotParameters();
+	return ScreenshotParameters;
 }
 
-bool  UInterchangeImportTestStepReimport::CanEditPipelineSettings() const
+bool UInterchangeImportTestStepReimport::CanEditPipelineSettings() const
 {
-	return false;
+	if (UInterchangeImportTestStepImport* ImportStep = CachedImportStep.Get())
+	{
+		if (!ImportStep->CanEditPipelineSettings())
+		{
+			return false;
+		}
+	}
+
+	if (bUseOverridePipelineStack)
+	{
+		if (PipelineStack.IsEmpty())
+		{
+			return false;
+		}
+
+		for (const UInterchangePipelineBase* Pipeline : PipelineStack)
+		{
+			if (!Pipeline)
+			{
+				return false;
+			}
+		}
+	}
+
+	return true;
 }
 
 void  UInterchangeImportTestStepReimport::EditPipelineSettings()
 {
+#if WITH_EDITOR
+	FString SourceFilePathString  = GetReimportStepSourceFilePathString();
+	if (SourceFilePathString.IsEmpty())
+	{
+		UE_LOG(LogInterchangeTests, Error, TEXT("Import Step doesn't have a valid Source File to get the default pipeline stack."));
+		return;
+	}
+
+	UE::Interchange::FScopedSourceData ScopedSourceData(SourceFilePathString);
+	UE::Interchange::FScopedTranslator ScopedTranslator(ScopedSourceData.GetSourceData());
+	if (!ScopedTranslator.GetTranslator())
+	{
+		UE_LOG(LogInterchangeTests, Error, TEXT("Cannot import file. The source data is not supported. Try enabling the [%s] extension for Interchange."), *FPaths::GetExtension(ScopedSourceData.GetSourceData()->GetFilename()));
+		return;
+	}
+
+	if (!CachedImportStep.IsValid())
+	{
+		UE_LOG(LogInterchangeTests, Error, TEXT("No valid import step found. Make sure this Reimport Step is part of an Interchange Test Plan Asset and is not used independently."));
+		return;
+	}
+
+	const bool bIsSceneImport = CachedImportStep.IsValid() && CachedImportStep->bImportIntoLevel;
+	const FInterchangeImportSettings& InterchangeImportSettings = FInterchangeProjectSettingsUtils::GetDefaultImportSettings(bIsSceneImport);
+	if (InterchangeImportSettings.PipelineStacks.Num() == 0)
+	{
+		UE_LOG(LogInterchangeTests, Error, TEXT("Failed to configure pipelines. There is no pipeline stack defined for the content import type."));
+		return;
+	}
+
+	if (!InterchangeImportSettings.PipelineStacks.Contains(InterchangeImportSettings.DefaultPipelineStack))
+	{
+		FInterchangeImportSettings& MutableInterchangeImportSettings = FInterchangeProjectSettingsUtils::GetMutableDefaultImportSettings(bIsSceneImport);
+		TArray<FName> Keys;
+		MutableInterchangeImportSettings.PipelineStacks.GetKeys(Keys);
+		check(Keys.Num() > 0);
+		MutableInterchangeImportSettings.DefaultPipelineStack = Keys[0];
+	}
+
+	UE::Interchange::FScopedBaseNodeContainer ScopedBaseNodeContainer;
+
+	UInterchangePipelineConfigurationBase* RegisteredPipelineConfiguration = nullptr;
+	TSoftClassPtr <UInterchangePipelineConfigurationBase> ImportDialogClass = InterchangeImportSettings.ImportDialogClass;
+	if (ImportDialogClass.IsValid())
+	{
+		UClass* PipelineConfigurationClass = ImportDialogClass.LoadSynchronous();
+		if (PipelineConfigurationClass)
+		{
+			RegisteredPipelineConfiguration = NewObject<UInterchangePipelineConfigurationBase>(GetTransientPackage(), PipelineConfigurationClass, NAME_None, RF_NoFlags);
+			if (RegisteredPipelineConfiguration == nullptr)
+			{
+				UE_LOG(LogInterchangeTests, Error, TEXT("Failed to create a pipeline configuration object."));
+				return;
+			}
+		}
+	}
+
+	auto AdjustPipelineSettingForContext = [this, &ScopedBaseNodeContainer, bIsSceneImport](UInterchangePipelineBase* Pipeline)
+		{
+			EInterchangePipelineContext Context;
+			Context = bIsSceneImport ? EInterchangePipelineContext::SceneReimport: EInterchangePipelineContext::AssetReimport;
+			FInterchangePipelineContextParams ContextParams;
+			ContextParams.ContextType = Context;
+			ContextParams.BaseNodeContainer = ScopedBaseNodeContainer.GetBaseNodeContainer();
+			Pipeline->SetFromReimportOrOverride(true);
+			Pipeline->AdjustSettingsForContext(ContextParams);
+		};
+
+	const bool bCanTranslate = UInterchangeManager::GetInterchangeManager().CanUseTranslator(ScopedTranslator.GetTranslator());
+	if (bCanTranslate)
+	{
+		FScopedSlowTask Progress(2.f, LOCTEXT("TranslatingSourceFile...", "Translating source file..."));
+		Progress.MakeDialog();
+		Progress.EnterProgressFrame(1.f);
+
+		UInterchangeBaseNodeContainer& BaseNodeContainer = *(ScopedBaseNodeContainer.GetBaseNodeContainer());
+		ScopedTranslator.GetTranslator()->SetResultsContainer(NewObject<UInterchangeResultsContainer>(GetTransientPackage(), NAME_None));
+		ScopedTranslator.GetTranslator()->Translate(BaseNodeContainer);
+
+		Progress.EnterProgressFrame(1.f);
+	}
+
+	const TMap<FName, FInterchangePipelineStack>& DefaultPipelineStacks = InterchangeImportSettings.PipelineStacks;
+	TArray<FInterchangeStackInfo> InPipelineStacks;
+	TArray<UInterchangePipelineBase*> OutPipelines;
+
+	if (bUseOverridePipelineStack)
+	{
+		if (!PipelineStack.IsEmpty())
+		{
+			FInterchangeStackInfo& StackInfo = InPipelineStacks.AddDefaulted_GetRef();
+			StackInfo.StackName = FName(TEXT("ReimportOverridePipelineStack"));
+			for (int32 PipelineClassIndex = 0; PipelineClassIndex < PipelineStack.Num(); ++PipelineClassIndex)
+			{
+				if (UInterchangePipelineBase* GeneratedPipeline = UE::Interchange::GeneratePipelineInstance(PipelineStack[PipelineClassIndex]))
+				{
+					AdjustPipelineSettingForContext(GeneratedPipeline);
+					StackInfo.Pipelines.Add(GeneratedPipeline);
+				}
+			}
+
+			if (StackInfo.Pipelines.IsEmpty())
+			{
+				UE_LOG(LogInterchangeTests, Error, TEXT("Failed to configure pipelines. There are no pipelines in the override stack"));
+				return;
+			}
+		}
+	}
+	else if (!PipelineSettings.CustomPipelines.IsEmpty())
+	{
+		FInterchangeStackInfo& StackInfo = InPipelineStacks.AddDefaulted_GetRef();
+		StackInfo.StackName = FName(TEXT("ReimportCustomPipelineStack"));
+		for (int32 PipelineIndex = 0; PipelineIndex < PipelineSettings.CustomPipelines.Num(); ++PipelineIndex)
+		{
+			if (UInterchangePipelineBase* GeneratedPipeline = UE::Interchange::GeneratePipelineInstance(PipelineSettings.CustomPipelines[PipelineIndex]))
+			{
+				AdjustPipelineSettingForContext(GeneratedPipeline);
+				StackInfo.Pipelines.Add(GeneratedPipeline);
+			}
+		}
+	}
+	else if (CachedImportStep.IsValid())
+	{
+		// Add the import step's pipeline stack as a starting point if it is using custom as most users would. User can always choose to go back to default stacks.
+		FInterchangeStackInfo& StackInfo = InPipelineStacks.AddDefaulted_GetRef();
+		StackInfo.StackName = FName(TEXT("ImportStepCustomPipelineStack"));
+
+		TArray<TObjectPtr<UInterchangePipelineBase>> ImportStepPipelines = CachedImportStep->GetCurrentPipelinesOrDefault();
+		for (int32 PipelineIndex = 0; PipelineIndex < ImportStepPipelines.Num(); ++PipelineIndex)
+		{
+			if (UInterchangePipelineBase* GeneratedPipeline = UE::Interchange::GeneratePipelineInstance(ImportStepPipelines[PipelineIndex]))
+			{
+				AdjustPipelineSettingForContext(GeneratedPipeline);
+				StackInfo.Pipelines.Add(GeneratedPipeline);
+			}
+		}
+	}
+
+	for (const TPair<FName, FInterchangePipelineStack>& PipelineStackInfo : DefaultPipelineStacks)
+	{
+		FName StackName = PipelineStackInfo.Key;
+		FInterchangeStackInfo& StackInfo = InPipelineStacks.AddDefaulted_GetRef();
+		StackInfo.StackName = StackName;
+
+		const FInterchangePipelineStack& DefPipelineStack = PipelineStackInfo.Value;
+		const TArray<FSoftObjectPath>* Pipelines = &DefPipelineStack.Pipelines;
+
+		// If applicable, check to see if a specific pipeline stack is associated with this translator
+		for (const FInterchangeTranslatorPipelines& TranslatorPipelines : DefPipelineStack.PerTranslatorPipelines)
+		{
+			const UClass* TranslatorClass = TranslatorPipelines.Translator.LoadSynchronous();
+			if (ScopedTranslator.GetTranslator() && ScopedTranslator.GetTranslator()->IsA(TranslatorClass))
+			{
+				Pipelines = &TranslatorPipelines.Pipelines;
+				break;
+			}
+		}
+
+		for (int32 PipelineIndex = 0; PipelineIndex < Pipelines->Num(); ++PipelineIndex)
+		{
+			if (UInterchangePipelineBase* GeneratedPipeline = UE::Interchange::GeneratePipelineInstance((*Pipelines)[PipelineIndex]))
+			{
+				AdjustPipelineSettingForContext(GeneratedPipeline);
+				StackInfo.Pipelines.Add(GeneratedPipeline);
+			}
+		}
+	}
+
+	//Show the dialog, a plugin should have register this dialog. We use a plugin to be able to use editor code when doing UI
+	constexpr bool bIsReimport = true;
+	EInterchangePipelineConfigurationDialogResult DialogResult = RegisteredPipelineConfiguration->ScriptedShowTestPlanConfigurationDialog(InPipelineStacks, OutPipelines, ScopedSourceData.GetSourceData(), ScopedTranslator.GetTranslator(), ScopedBaseNodeContainer.GetBaseNodeContainer(), nullptr, bIsSceneImport, bIsReimport);
+
+	if (DialogResult == EInterchangePipelineConfigurationDialogResult::SaveConfig)
+	{
+		if (bUseOverridePipelineStack)
+		{
+			FScopedTransaction ScopedTransaction(LOCTEXT("UpdatePipelineSettings", "Update Pipeline Settings"));
+			Modify();
+
+			PipelineStack.Empty(OutPipelines.Num());
+			for (UInterchangePipelineBase* Pipeline : OutPipelines)
+			{
+				if (UInterchangePipelineBase* InstancedPipeline = DuplicateObject<UInterchangePipelineBase>(Pipeline, this))
+				{
+					InstancedPipeline->SetFlags(EObjectFlags::RF_Transactional);
+					PipelineStack.Add(InstancedPipeline);
+				}
+			}
+
+		}
+		else
+		{
+			PipelineSettings.UpdatePipelines(OutPipelines);
+		}
+	}
+#endif	
 }
 
 void  UInterchangeImportTestStepReimport::ClearPipelineSettings()
 {
+
+#if WITH_EDITOR
+	if (bUseOverridePipelineStack)
+	{
+		FScopedTransaction ScopedTransaction(LOCTEXT("ClearPipelineSettings", "Clear Pipeline Settings"));
+		Modify();
+		PipelineStack.Empty();
+	}
+	else
+	{
+		PipelineSettings.ClearPipelines();
+	}
+#endif
+
 }
+
+bool UInterchangeImportTestStepReimport::IsUsingOverridePipelines(bool bCheckForValidPipelines) const
+{
+	if (!bCheckForValidPipelines)
+	{
+		return bUseOverridePipelineStack;
+	}
+
+	if (bUseOverridePipelineStack)
+	{
+		for (const TObjectPtr<UInterchangePipelineBase>& Pipeline : PipelineStack)
+		{
+			if (!Pipeline)
+			{
+				return false;
+			}
+		}
+		return true;
+	}
+
+	return false;
+}
+
+void UInterchangeImportTestStepReimport::HandleImportPipelineSettingsModified(FImportStepChangedData ChangedData)
+{
+	if (!bUseOverridePipelineStack && !PipelineSettings.CustomPipelines.IsEmpty())
+	{
+		FSuppressableWarningDialog::FSetupInfo DialogSetupInfo = FSuppressableWarningDialog::FSetupInfo(
+			LOCTEXT("ClearReimportPipelinesDialogText", "You have pipeline settings incompatible with the import step.\nClearing out to restore correct defaults."),
+			LOCTEXT("ClearReimportPipelinesDialogTitle", "Incompatible Pipeline Settings"),
+			"InterchangeImportTestClearReimportPipelinesWarning");
+		
+		DialogSetupInfo.ConfirmText = LOCTEXT("ClearReimportPipelinesDialogTextCancel", "OK");
+		DialogSetupInfo.bDefaultToSuppressInTheFuture = false;
+		DialogSetupInfo.CheckBoxText = LOCTEXT("ClearReimportPipelinesDialogCheckBoxText", "Don't show this dialog again");
+		DialogSetupInfo.DialogMode = UInterchangeImportTestPlan::GetInterchangeTestPlanWarningDialogMode();
+		FSuppressableWarningDialog SuppressableWarningDialog(DialogSetupInfo);
+		SuppressableWarningDialog.ShowModal();
+
+		PipelineSettings.ClearPipelines(ChangedData.ChangeType == EImportStepDataChangeType::ImportIntoLevel);
+	}
+}
+
+FString UInterchangeImportTestStepReimport::GetReimportStepSourceFilePathString()
+{
+	if (CachedImportStep.IsValid())
+	{
+		if (UInterchangeImportTestStepImport* ImportStep = CachedImportStep.Get())
+		{
+			if (!ImportStep->SourceFile.FilePath.IsEmpty())
+			{
+				if (!SourceFileToReimport.FilePath.IsEmpty())
+				{
+					return SourceFileToReimport.FilePath;
+				}
+
+				return ImportStep->SourceFile.FilePath;
+			}
+		}
+	}
+
+	return FString();
+}
+
+#if WITH_EDITOR
+void UInterchangeImportTestStepReimport::PostEditChangeProperty(FPropertyChangedEvent& PropertyChangedEvent)
+{
+	static const FName UseOverridePipelineStackPropertyName = GET_MEMBER_NAME_CHECKED(UInterchangeImportTestStepImport, bUseOverridePipelineStack);
+
+	if (PropertyChangedEvent.Property->GetFName() == UseOverridePipelineStackPropertyName)
+	{
+		if (bUseOverridePipelineStack && !PipelineSettings.CustomPipelines.IsEmpty())
+		{
+			FSuppressableWarningDialog::FSetupInfo DialogSetupInfo = FSuppressableWarningDialog::FSetupInfo(
+				LOCTEXT("UseOverridePipelinesDialogText", "You are using override pipelines now. This will erase all the previous modifications made to the default pipeline settings."),
+				LOCTEXT("UseOverridePipelinesDialogTitle", "Using Override Pipelines"),
+				"InterchangeImportTestPlanUseOverridePipelinesWarning");
+
+			DialogSetupInfo.ConfirmText = LOCTEXT("UseOverridePipelinesDialogOptionConfirm", "OK");
+			DialogSetupInfo.bDefaultToSuppressInTheFuture = false;
+			DialogSetupInfo.CheckBoxText = LOCTEXT("UseOverridePipelinesDialogCheckBoxText", "Don't show this dialog again");
+			DialogSetupInfo.DialogMode = UInterchangeImportTestPlan::GetInterchangeTestPlanWarningDialogMode();
+			FSuppressableWarningDialog SuppressableWarningDialog(DialogSetupInfo);
+			SuppressableWarningDialog.ShowModal();
+
+			PipelineSettings.ClearPipelines();
+		}
+	}
+
+	Super::PostEditChangeProperty(PropertyChangedEvent);
+}
+
+void UInterchangeImportTestStepReimport::PostEditChangeChainProperty(FPropertyChangedChainEvent& PropertyChangedEvent)
+{
+	static const FName SourceFileToReimportPropertyName = GET_MEMBER_NAME_CHECKED(UInterchangeImportTestStepReimport, SourceFileToReimport);
+
+	const FProperty* ActiveMemberPropertyNode = PropertyChangedEvent.PropertyChain.GetActiveMemberNode()->GetValue();
+	if (ActiveMemberPropertyNode && (ActiveMemberPropertyNode->GetFName() == SourceFileToReimportPropertyName))
+	{
+		const FString CurrSourceFileExtension = FPaths::GetExtension(SourceFileToReimport.FilePath);
+		if (!LastSourceFileExtension.IsEmpty() && !CurrSourceFileExtension.IsEmpty()
+			&& LastSourceFileExtension != CurrSourceFileExtension
+			&& !bUseOverridePipelineStack
+			&& !PipelineSettings.CustomPipelines.IsEmpty())
+		{
+			FSuppressableWarningDialog::FSetupInfo DialogSetupInfo = FSuppressableWarningDialog::FSetupInfo(
+				LOCTEXT("SourceFileExtensionChangedDialogText", "Current pipelines might not be compatible with the new source file. Clearing out to restore the defaults."),
+				LOCTEXT("SourceFileExtensionChangedDialogTitle", "Source File Extension Changed"),
+				"InterchangeImportTestPlanSourceFileExtensionChangedWarning");
+
+			DialogSetupInfo.ConfirmText = LOCTEXT("SourceFileExtensionChangedDialogOptionConfirm", "OK");
+			DialogSetupInfo.bDefaultToSuppressInTheFuture = false;
+			DialogSetupInfo.CheckBoxText = LOCTEXT("SourceFileExtensionChangedDialogCheckBoxText", "Don't show this dialog again");
+			DialogSetupInfo.DialogMode = UInterchangeImportTestPlan::GetInterchangeTestPlanWarningDialogMode();
+			FSuppressableWarningDialog SuppressableWarningDialog(DialogSetupInfo);
+			SuppressableWarningDialog.ShowModal();
+
+			PipelineSettings.ClearPipelines();
+		}
+		LastSourceFileExtension = CurrSourceFileExtension;
+	}
+	
+	Super::PostEditChangeChainProperty(PropertyChangedEvent);
+}
+#endif
+
+#undef LOCTEXT_NAMESPACE 
\ No newline at end of file
diff --git a/Engine/Plugins/Tests/InterchangeTests/Source/InterchangeTests/Private/InterchangeTestPlanPipelineSettings.cpp b/Engine/Plugins/Tests/InterchangeTests/Source/InterchangeTests/Private/InterchangeTestPlanPipelineSettings.cpp
--- a/Engine/Plugins/Tests/InterchangeTests/Source/InterchangeTests/Private/InterchangeTestPlanPipelineSettings.cpp
+++ b/Engine/Plugins/Tests/InterchangeTests/Source/InterchangeTests/Private/InterchangeTestPlanPipelineSettings.cpp
@@ -4,7 +4,7 @@
 #include "InterchangePipelineBase.h"
 #include "ScopedTransaction.h"
 
-void FInterchangeTestPlanPipelineSettings::UpdatePipelines(const TArray<UInterchangePipelineBase*>& InPipelines)
+void FInterchangeTestPlanPipelineSettings::UpdatePipelines(const TArray<UInterchangePipelineBase*>& InPipelines, bool bTransactional /*= true*/)
 {
 	if (ensure(ParentTestStep))
 	{
@@ -13,18 +13,45 @@ void FInterchangeTestPlanPipelineSettings::UpdatePipelines(const TArray<UInterch
 		CustomPipelines.Empty(InPipelines.Num());
 		for (const UInterchangePipelineBase* Pipeline : InPipelines)
 		{
-			CustomPipelines.Emplace(DuplicateObject<UInterchangePipelineBase>(Pipeline, ParentTestStep->GetOuter()));
+			CustomPipelines.Emplace(DuplicateObject<UInterchangePipelineBase>(Pipeline, ParentTestStep));
+		}
+		if (!bTransactional)
+		{
+			ScopedTransaction.Cancel();
 		}
 	}
 }
 
-void FInterchangeTestPlanPipelineSettings::ClearPipelines()
+void FInterchangeTestPlanPipelineSettings::UpdatePipelines(const TArray<TObjectPtr<UInterchangePipelineBase>>& InPipelines, bool bTransactional /*= true*/)
 {
 	if (ensure(ParentTestStep))
+	{
+		FScopedTransaction ScopedTransaction(NSLOCTEXT("InterchangeTestPlanPipelineSettings", "UpdatePipelineSettings", "Update Pipeline Settings"));
+		ParentTestStep->Modify();
+
+		CustomPipelines.Empty(InPipelines.Num());
+		for (const TObjectPtr<UInterchangePipelineBase>& Pipeline : InPipelines)
+		{
+			CustomPipelines.Emplace(DuplicateObject<UInterchangePipelineBase>(Pipeline, ParentTestStep));
+		}
+		if (!bTransactional)
+		{
+			ScopedTransaction.Cancel();
+		}
+	}
+}
+
+void FInterchangeTestPlanPipelineSettings::ClearPipelines(bool bTransactional /*= true*/)
+{
+	if (ensure(ParentTestStep) && !CustomPipelines.IsEmpty())
 	{
 		FScopedTransaction ScopedTransaction(NSLOCTEXT("InterchangeTestPlanPipelineSettings", "ClearPipelineSettings", "Clear Pipeline Settings"));
 		ParentTestStep->Modify();
 		CustomPipelines.Empty();
+		if (!bTransactional)
+		{
+			ScopedTransaction.Cancel();
+		}
 	}
 }
 
diff --git a/Engine/Plugins/Tests/InterchangeTests/Source/InterchangeTests/Public/InterchangeImportTestData.h b/Engine/Plugins/Tests/InterchangeTests/Source/InterchangeTests/Public/InterchangeImportTestData.h
--- a/Engine/Plugins/Tests/InterchangeTests/Source/InterchangeTests/Public/InterchangeImportTestData.h
+++ b/Engine/Plugins/Tests/InterchangeTests/Source/InterchangeTests/Public/InterchangeImportTestData.h
@@ -8,11 +8,36 @@ class UInterchangeImportTestPlan;
 class UInterchangeResultsContainer;
 
 
-struct FInterchangeImportTestData
+struct FInterchangeImportTestData : public FAssetData
 {
+public:
+	FInterchangeImportTestData()
+		: FAssetData()
+	{}
+
+	FInterchangeImportTestData(FName InPackageName, FName InPackagePath, FName InAssetName, FTopLevelAssetPath InAssetClassPathName, FAssetDataTagMap InTags = FAssetDataTagMap(), TArrayView<const int32> InChunkIDs = TArrayView<const int32>(), uint32 InPackageFlags = 0)
+		: FAssetData(InPackageName, InPackagePath, InAssetName, InAssetClassPathName, InTags, InChunkIDs, InPackageFlags)
+	{}
+	
+	FInterchangeImportTestData(const FString& InLongPackageName, const FString& InObjectPath, FTopLevelAssetPath InAssetClassPathName, FAssetDataTagMap InTags = FAssetDataTagMap(), TArrayView<const int32> InChunkIDs = TArrayView<const int32>(), uint32 InPackageFlags = 0)
+		: FAssetData(InLongPackageName, InObjectPath, InAssetClassPathName, InTags, InChunkIDs, InPackageFlags)
+	{}
+
+	FInterchangeImportTestData(const UObject* InAsset, FAssetData::ECreationFlags InCreationFlags = ECreationFlags::None)
+		: FAssetData(InAsset, InCreationFlags)
+	{}
+	
+	FInterchangeImportTestData(const UObject* InAsset, FAssetData::ECreationFlags InCreationFlags, EAssetRegistryTagsCaller Caller)
+		: FAssetData(InAsset, InCreationFlags, Caller)
+	{}
+
+	inline FInterchangeImportTestData(const UObject* InAsset, bool bAllowBlueprintClass)
+		: FInterchangeImportTestData(InAsset, bAllowBlueprintClass ? ECreationFlags::AllowBlueprintClass : ECreationFlags::None)
+	{
+	}
+
 	FString DestAssetPackagePath;
 	FString DestAssetFilePath;
-	FAssetData AssetData;
 	UInterchangeImportTestPlan* TestPlan = nullptr;
 	UInterchangeResultsContainer* InterchangeResults = nullptr;
 	TArray<UObject*> ResultObjects;
diff --git a/Engine/Plugins/Tests/InterchangeTests/Source/InterchangeTests/Public/InterchangeImportTestPlan.h b/Engine/Plugins/Tests/InterchangeTests/Source/InterchangeTests/Public/InterchangeImportTestPlan.h
--- a/Engine/Plugins/Tests/InterchangeTests/Source/InterchangeTests/Public/InterchangeImportTestPlan.h
+++ b/Engine/Plugins/Tests/InterchangeTests/Source/InterchangeTests/Public/InterchangeImportTestPlan.h
@@ -2,15 +2,20 @@
 
 #pragma once
 
+#include "Engine/World.h"
+#include "Dialogs/Dialogs.h"
+#include "UObject/ObjectPtr.h"
+
 #include "InterchangeImportTestPlan.generated.h"
 
 class UInterchangeImportTestStepBase;
-
+class UInterchangeImportTestStepImport;
+class UInterchangeImportTestStepReimport;
 
 /**
 * Define a test plan
 */
-UCLASS(BlueprintType, EditInlineNew, PrioritizeCategories = Definition) // @TODO: can't find a way to force Run to be below Definition
+UCLASS(BlueprintType, EditInlineNew, meta=(PrioritizeCategories = "Description Level Definition Import Reimport")) // @TODO: can't find a way to force Run to be below Definition
 class INTERCHANGETESTS_API UInterchangeImportTestPlan : public UObject
 {
 	GENERATED_BODY()
@@ -24,18 +29,68 @@ public:
 	void WriteToJson(const FString& Filename);
 
 public:
+	UInterchangeImportTestPlan();
+
+	/** Test Description */
+	UPROPERTY(EditAnywhere, Category = "Description", Meta = (MultiLine = true))
+	FString Description;
+
 	/** Set of steps to perform to carry out this test plan */
-	UPROPERTY(EditAnywhere, Instanced, Category = Definition, Meta = (DisplayPriority = 0))
-	TArray<TObjectPtr<UInterchangeImportTestStepBase>> Steps;
+	UPROPERTY(Instanced, meta = (DeprecatedProperty, DeprecationMessage = "Use the import and reimport sections to fill up the test steps."))
+	TArray<TObjectPtr<UInterchangeImportTestStepBase>> Steps_DEPRECATED;
+	
+	/**  Level to use for taking Screenshot */
+	UPROPERTY(EditAnywhere, Category = Level, meta = (AllowedClasses = "/Script/Engine.World"))
+	FSoftObjectPath WorldPath;
 
-	/** Click here to immediately run this single test through the automation framework */
-	UFUNCTION(BlueprintCallable, CallInEditor, Category = Run, Meta = (DisplayPriority = 1))
-	void RunThisTest();
+	/** File Import Step */
+	UPROPERTY(VisibleAnywhere, Instanced, Category = "Import")
+	TObjectPtr<UInterchangeImportTestStepImport> ImportStep;
 
+	/** Set of re-imports that will follow the above import step */
+	UPROPERTY(EditAnywhere, Instanced, Category = "Reimport", NoClear)
+	TArray<TObjectPtr<UInterchangeImportTestStepReimport>> ReimportStack;
+
+
+public:
+	/** Click here to immediately run this single test through the automation framework */
+	UFUNCTION(BlueprintCallable, CallInEditor, Category = "Automation")
+	void RunThisTest();	
 	bool IsRunningSynchornously() const;
 
+	void SetupLevelForImport();
+	
+	UWorld* GetCurrentWorld() const
+	{	
+		return TransientWorld.Get() ? TransientWorld.Get() : GWorld->GetWorld();
+	}
+
+	ULevel* GetCurrentLevel() const 
+	{
+		return GetCurrentWorld()->GetCurrentLevel();
+	}
+
+	void CleanupLevel();
+
+	virtual void PostLoad() override;
+	void OnAssetLoaded(UObject* Asset);
+
+#if WITH_EDITOR
+	static FSuppressableWarningDialog::EMode GetInterchangeTestPlanWarningDialogMode();
+
+	virtual void PreEditUndo() override;
+	virtual void PostEditUndo() override;
+
+	virtual void PreEditChange(FProperty* PropertyAboutToChange) override;
+	virtual void PostEditChangeChainProperty(FPropertyChangedChainEvent& PropertyChangedEvent) override;
+#endif
+	
 private:
+	TStrongObjectPtr<UWorld> TransientWorld;
+
+	int32 PrevReimportStackSize = -1;
 	bool bRunSynchronously = false;
+	bool bChangeObjectOuters = false;
 };
 
 
diff --git a/Engine/Plugins/Tests/InterchangeTests/Source/InterchangeTests/Public/InterchangeImportTestStepBase.h b/Engine/Plugins/Tests/InterchangeTests/Source/InterchangeTests/Public/InterchangeImportTestStepBase.h
--- a/Engine/Plugins/Tests/InterchangeTests/Source/InterchangeTests/Public/InterchangeImportTestStepBase.h
+++ b/Engine/Plugins/Tests/InterchangeTests/Source/InterchangeTests/Public/InterchangeImportTestStepBase.h
@@ -6,9 +6,10 @@
 #include "AutomationScreenshotOptions.h"
 #include "Containers/EnumAsByte.h"
 #include "Engine/EngineBaseTypes.h"
+#include "UObject/ObjectPtr.h"
 #include "InterchangeImportTestStepBase.generated.h"
 
-class ACameraActor;
+class UInterchangeImportTestPlan;
 class FAutomationTestExecutionInfo;
 struct FInterchangeImportTestData;
 struct FInterchangeTestFunction;
@@ -61,7 +62,8 @@ public:
 	/** An array of results to check against */
 	UPROPERTY(EditAnywhere, Category = General)
 	TArray<FInterchangeTestFunction> Tests;
-
+	
+	TObjectPtr<UInterchangeImportTestPlan> ParentTestPlan = nullptr;
 public:
 	virtual TTuple<UE::Interchange::FAssetImportResultPtr, UE::Interchange::FSceneImportResultPtr>
 		StartStep(FInterchangeImportTestData& Data) PURE_VIRTUAL(UInterchangeImportTestStepBase::StartStep, return {}; );
@@ -73,6 +75,7 @@ public:
 	virtual bool CanEditPipelineSettings() const PURE_VIRTUAL(UInterchangeImportTestStepBase::CanEditPipelineSettings, return false; );
 	virtual void EditPipelineSettings() PURE_VIRTUAL(UInterchangeImportTestStepBase::EditPipelineSettings);
 	virtual void ClearPipelineSettings() PURE_VIRTUAL(UInterchangeImportTestStepBase::ClearPipelineSettings);
+	virtual bool IsUsingOverridePipelines(bool bCheckForValidPipelines) const PURE_VIRTUAL(UInterchangeImportTestStepBase::IsUsingOverridePipelines, return false; );
 
 	bool PerformTests(FInterchangeImportTestData& Data, FAutomationTestBase* CurrentTest);
 protected:
diff --git a/Engine/Plugins/Tests/InterchangeTests/Source/InterchangeTests/Public/InterchangeImportTestStepImport.h b/Engine/Plugins/Tests/InterchangeTests/Source/InterchangeTests/Public/InterchangeImportTestStepImport.h
--- a/Engine/Plugins/Tests/InterchangeTests/Source/InterchangeTests/Public/InterchangeImportTestStepImport.h
+++ b/Engine/Plugins/Tests/InterchangeTests/Source/InterchangeTests/Public/InterchangeImportTestStepImport.h
@@ -2,22 +2,47 @@
 
 #pragma once
 
+#include "Delegates/DelegateCombinations.h"
+#include "InterchangeImportTestPlan.h"
 #include "InterchangeImportTestStepBase.h"
-#include "UObject/StrongObjectPtr.h"
 #include "InterchangeTestPlanPipelineSettings.h"
+#include "UObject/StrongObjectPtr.h"
 
 #include "InterchangeImportTestStepImport.generated.h"
 
+class UInterchangeTranslatorBase;
 class UInterchangePipelineBase;
+class UInterchangeBaseNodeContainer;
 class UWorld;
 
-UCLASS(BlueprintType, AutoExpandCategories = ("General", "Screenshot Comparison"), Meta = (DisplayName = "Import a file"))
+class UInterchangeImportTestStepImport;
+
+enum class EImportStepDataChangeType : uint8
+{
+	Unknown,
+	SourceFile,
+	PipelineSettings,
+	ImportIntoLevel,
+};
+
+struct FImportStepChangedData
+{
+public:
+	EImportStepDataChangeType ChangeType;
+
+	UInterchangeImportTestStepImport* ImportStep = nullptr;
+};
+
+DECLARE_MULTICAST_DELEGATE_OneParam(FOnImportTestStepDataChanged, FImportStepChangedData)
+
+
+UCLASS(BlueprintType, Meta = (DisplayName = "Import a file"))
 class INTERCHANGETESTS_API UInterchangeImportTestStepImport : public UInterchangeImportTestStepBase
 {
 	GENERATED_BODY()
 
 public:
-	UInterchangeImportTestStepImport(const FObjectInitializer& InObjectInitializer);
+	UInterchangeImportTestStepImport();
 
 	/** The source file to import (path relative to the json script) */
 	UPROPERTY(EditAnywhere, Category = General)
@@ -38,21 +63,19 @@ public:
 	UPROPERTY(EditAnywhere, Category = General)
 	bool bEmptyDestinationFolderPriorToImport = true;
 
-	/** Whether imported assets should be saved and freshly reloaded after import */
-	UPROPERTY(EditAnywhere, Category = General, meta=(EditCondition="!bImportIntoLevel"))
-	bool bSaveThenReloadImportedAssets = true;
-
 	/**  Whether we should use the import into level workflow */
 	UPROPERTY(EditAnywhere, Category = General)
 	bool bImportIntoLevel = false;
 
-	/**  Screen Shot Comparison when Importing in Level */
-	UPROPERTY(EditAnywhere, Category = "Screenshot Comparison", meta = (AllowedClasses = "/Script/Engine.World", EditCondition = "bImportIntoLevel"))
-	FSoftObjectPath WorldPath;
+	/**  Whether screenshot would be taken at this stage. */
+	UPROPERTY(EditAnywhere, Category = "Screenshot Comparison")
+	bool bTakeScreenshot = false;
 
-	UPROPERTY(EditAnywhere, Category = "Screenshot Comparison", meta = (EditCondition = "bImportIntoLevel"))
+	/**  Screen Shot Settings */
+	UPROPERTY(EditAnywhere, Category = "Screenshot Comparison", meta = (EditCondition = "bTakeScreenshot"))
 	FInterchangeTestScreenshotParameters ScreenshotParameters;
 
+	FOnImportTestStepDataChanged OnImportTestStepDataChanged;
 public:
 	// UInterchangeImportTestStepBase interface
 	virtual TTuple<UE::Interchange::FAssetImportResultPtr, UE::Interchange::FSceneImportResultPtr>
@@ -65,13 +88,23 @@ public:
 	virtual bool CanEditPipelineSettings() const override;
 	virtual void EditPipelineSettings() override;
 	virtual void ClearPipelineSettings() override;
+	virtual bool IsUsingOverridePipelines(bool bCheckForValidPipelines) const override;
 
 	bool ShouldImportIntoLevelChangeRequireMessageBox() const;
 
+	virtual void PostLoad() override;
+
 #if WITH_EDITOR
 	virtual void PostEditChangeProperty(FPropertyChangedEvent& PropertyChangedEvent) override;
+	virtual void PostEditChangeChainProperty(FPropertyChangedChainEvent& PropertyChangedEvent) override;
 #endif
+	
+	TArray<TObjectPtr<UInterchangePipelineBase>> GetCurrentPipelinesOrDefault() const;
+private:
+	void BroadcastImportStepChangedEvent(EImportStepDataChangeType ChangeType);
 
 private:
-	TStrongObjectPtr<UWorld> TransientWorld;
+	UPROPERTY()
+	FString LastSourceFileExtension;
+
 };
diff --git a/Engine/Plugins/Tests/InterchangeTests/Source/InterchangeTests/Public/InterchangeImportTestStepReimport.h b/Engine/Plugins/Tests/InterchangeTests/Source/InterchangeTests/Public/InterchangeImportTestStepReimport.h
--- a/Engine/Plugins/Tests/InterchangeTests/Source/InterchangeTests/Public/InterchangeImportTestStepReimport.h
+++ b/Engine/Plugins/Tests/InterchangeTests/Source/InterchangeTests/Public/InterchangeImportTestStepReimport.h
@@ -4,17 +4,24 @@
 
 #include "InterchangeImportTestStepBase.h"
 #include "Templates/SubclassOf.h"
+#include "UObject/ObjectPtr.h"
+#include "InterchangeTestPlanPipelineSettings.h"
+#include "InterchangeImportTestStepImport.h"
+#include "Delegates/DelegateCombinations.h"
 #include "InterchangeImportTestStepReimport.generated.h"
 
+class UInterchangeImportTestPlan;
 class UInterchangePipelineBase;
 struct FInterchangeTestFunction;
 
-UCLASS(BlueprintType, Meta = (DisplayName = "Reimport a file"))
+UCLASS(BlueprintType, Meta = (DisplayName = "Reimport a file"), AutoExpandCategories="General")
 class INTERCHANGETESTS_API UInterchangeImportTestStepReimport : public UInterchangeImportTestStepBase
 {
 	GENERATED_BODY()
 
 public:
+	UInterchangeImportTestStepReimport();
+
 	/** The source file to import (path relative to the json script). */
 	UPROPERTY(EditAnywhere, Category = General)
 	FFilePath SourceFileToReimport;
@@ -24,22 +31,37 @@ public:
 	bool bUseOverridePipelineStack = false;
 
 	/** The pipeline stack to use when re-importing (an empty array will use the original import pipelines) */
-	UPROPERTY(EditAnywhere, Instanced, Category = General, meta = (DisplayName = "Override Pipeline Stack", EditCondition = "bUseOverridePipelineStack"))
+	UPROPERTY(EditAnywhere, Instanced, Category = General, meta = (DisplayName = "Override Pipeline Stack", EditCondition = "bUseOverridePipelineStack", MaxPropertyDepth = 1))
 	TArray<TObjectPtr<UInterchangePipelineBase>> PipelineStack;
 
-	/** The type of the asset to reimport. If only one such asset was imported, this is unambiguous. */
+	/** Pipeline settings that would allow modifying the pipelines as reimport pipelines */
+	UPROPERTY(EditAnywhere, Category = General)
+	FInterchangeTestPlanPipelineSettings PipelineSettings;
+
+	/** If this is an import into level with new file in the same directory as import */
 	UPROPERTY(EditAnywhere, Category = General)
+	bool bImportIntoLevel = false;
+
+	/** The type of the asset to reimport. If only one such asset was imported, this is unambiguous. */
+	UPROPERTY(EditAnywhere, Category = General, meta=(EditCondition = "!bImportIntoLevel"))
 	TSubclassOf<UObject> AssetTypeToReimport;
 
 	/** If there were multiple assets of the above type imported, specify the concrete name here. */
-	UPROPERTY(EditAnywhere, Category = General)
+	UPROPERTY(EditAnywhere, Category = General, meta = (EditCondition = "!bImportIntoLevel"))
 	FString AssetNameToReimport;
 
-	/** Determines whether imported assets should be saved and freshly reloaded after import. */
-	UPROPERTY(EditAnywhere, Category = General)
-	bool bSaveThenReloadImportedAssets = true;
+	/** Whether a screenshot should be captured after reimport. */
+	UPROPERTY(EditAnywhere, Category = "Screenshot Comparison")
+	bool bTakeScreenshot = false;
+
+	/**  Screen Shot Settings */
+	UPROPERTY(EditAnywhere, Category = "Screenshot Comparison", meta=(EditCondition="bTakeScreenshot"))
+	FInterchangeTestScreenshotParameters ScreenshotParameters;
 
 public:
+	void InitializeReimportStepFromImportStep(UInterchangeImportTestStepImport* ImportTestStep);
+	void RemoveImportStepPipelineSettingsModifiedDelegate();
+
 	// UInterchangeImportTestStepBase interface
 	virtual TTuple<UE::Interchange::FAssetImportResultPtr, UE::Interchange::FSceneImportResultPtr>
 		StartStep(FInterchangeImportTestData& Data) override;
@@ -51,5 +73,20 @@ public:
 	virtual bool CanEditPipelineSettings() const override;
 	virtual void EditPipelineSettings() override;
 	virtual void ClearPipelineSettings() override;
+	virtual bool IsUsingOverridePipelines(bool bCheckForValidPipelines) const override;
+
+	void HandleImportPipelineSettingsModified(FImportStepChangedData ChangedData);
+
+	FString GetReimportStepSourceFilePathString();
+
+#if WITH_EDITOR
+	virtual void PostEditChangeProperty(FPropertyChangedEvent& PropertyChangedEvent) override;
+	virtual void PostEditChangeChainProperty(FPropertyChangedChainEvent& PropertyChangedEvent) override;
+#endif
+
+private:
+	UPROPERTY()
+	FString LastSourceFileExtension;
 
+	TWeakObjectPtr<UInterchangeImportTestStepImport> CachedImportStep = nullptr;
 };
diff --git a/Engine/Plugins/Tests/InterchangeTests/Source/InterchangeTests/Public/InterchangeTestPlanPipelineSettings.h b/Engine/Plugins/Tests/InterchangeTests/Source/InterchangeTests/Public/InterchangeTestPlanPipelineSettings.h
--- a/Engine/Plugins/Tests/InterchangeTests/Source/InterchangeTests/Public/InterchangeTestPlanPipelineSettings.h
+++ b/Engine/Plugins/Tests/InterchangeTests/Source/InterchangeTests/Public/InterchangeTestPlanPipelineSettings.h
@@ -1,6 +1,8 @@
 // Copyright Epic Games, Inc. All Rights Reserved.
 #pragma once
+
 #include "Delegates/DelegateCombinations.h"
+#include "UObject/ObjectPtr.h"
 #include "InterchangeTestPlanPipelineSettings.generated.h"
 
 class UInterchangeImportTestStepBase;
@@ -11,9 +13,10 @@ struct FInterchangeTestPlanPipelineSettings
 {
 	GENERATED_BODY()
 public:
-	void UpdatePipelines(const TArray<UInterchangePipelineBase*>& InPipelines);
+	void UpdatePipelines(const TArray<UInterchangePipelineBase*>& InPipelines, bool bTransactional = true);
+	void UpdatePipelines(const TArray<TObjectPtr<UInterchangePipelineBase>>& InPipelines, bool bTransactional = true);
 	
-	void ClearPipelines();
+	void ClearPipelines(bool bTransactional = true);
 
 	INTERCHANGETESTS_API bool IsUsingOverridePipelineStack() const;
 
