[Interchange Import Test Plan] Added more edit conditions to avoid confusion about ability to take screenshots

- Import step's screenshot parameters are disabled until all conditions are satisfied.

#rnx

diff --git a/Engine/Plugins/Tests/InterchangeTests/Source/InterchangeTests/Public/InterchangeImportTestStepImport.h b/Engine/Plugins/Tests/InterchangeTests/Source/InterchangeTests/Public/InterchangeImportTestStepImport.h
--- a/Engine/Plugins/Tests/InterchangeTests/Source/InterchangeTests/Public/InterchangeImportTestStepImport.h
+++ b/Engine/Plugins/Tests/InterchangeTests/Source/InterchangeTests/Public/InterchangeImportTestStepImport.h
@@ -68,11 +68,11 @@ public:
 	bool bImportIntoLevel = false;
 
 	/**  Whether screenshot would be taken at this stage. */
-	UPROPERTY(EditAnywhere, Category = "Screenshot Comparison")
+	UPROPERTY(EditAnywhere, Category = "Screenshot Comparison", meta=(EditCondition = "bImportIntoLevel"))
 	bool bTakeScreenshot = false;
 
 	/**  Screen Shot Settings */
-	UPROPERTY(EditAnywhere, Category = "Screenshot Comparison", meta = (EditCondition = "bTakeScreenshot"))
+	UPROPERTY(EditAnywhere, Category = "Screenshot Comparison", meta = (EditCondition = "bTakeScreenshot && bImportIntoLevel"))
 	FInterchangeTestScreenshotParameters ScreenshotParameters;
 
 	FOnImportTestStepDataChanged OnImportTestStepDataChanged;
