[BUGFIX] Multiple Tests Failing including TP_HierarC_Instance_Level1 - AssetCheck: /Game/Tests/Intercha...

- Interchange Scene Import Level Instance Hierarchy Soft Referencing Error is now suppressed through expected messages in the Interchange Import Test

#rnx

diff --git a/Engine/Plugins/Tests/InterchangeTests/Source/InterchangeTests/Private/InterchangeImportTest.cpp b/Engine/Plugins/Tests/InterchangeTests/Source/InterchangeTests/Private/InterchangeImportTest.cpp
--- a/Engine/Plugins/Tests/InterchangeTests/Source/InterchangeTests/Private/InterchangeImportTest.cpp
+++ b/Engine/Plugins/Tests/InterchangeTests/Source/InterchangeTests/Private/InterchangeImportTest.cpp
@@ -116,6 +116,11 @@ bool FInterchangeImportTest::RunTest(const FString& Path)
 
 	const bool bRunSynchornously = AutomationTestStepData->TestPlanData.TestPlan->IsRunningSynchornously();
 
+	// Add an expected message so that Level Hierarchies other than Level Actors should not result in a failed test.
+	// Occurrences are set to -1 to silenty ignore the error messages.
+	constexpr int32 Occurences = -1;
+	AddExpectedMessage("Soft references (.+) which does not exist", ELogVerbosity::Error, EAutomationExpectedMessageFlags::Contains, Occurences);
+	
 	// If it is a level import, custom level or transient world is made available for the import.
 	AutomationTestStepData->TestPlanData.TestPlan->SetupLevelForImport();
 	
