Updated headers using UnrealCodeFixup to make sure dllstorage is on methods/staticvars instead of types. Part 4/n

diff --git a/Engine/Plugins/Slate/SlateScripting/Source/SlateScriptingCommands/Public/UICommandsScriptingSubsystem.h b/Engine/Plugins/Slate/SlateScripting/Source/SlateScriptingCommands/Public/UICommandsScriptingSubsystem.h
--- a/Engine/Plugins/Slate/SlateScripting/Source/SlateScriptingCommands/Public/UICommandsScriptingSubsystem.h
+++ b/Engine/Plugins/Slate/SlateScripting/Source/SlateScriptingCommands/Public/UICommandsScriptingSubsystem.h
@@ -12,6 +12,8 @@
 
 #include "UICommandsScriptingSubsystem.generated.h"
 
+#define UE_API SLATESCRIPTINGCOMMANDS_API
+
 
 DECLARE_LOG_CATEGORY_EXTERN(LogCommandsScripting, Log, All);
 
@@ -174,8 +176,8 @@ struct FScriptingCommandSet
  * UEditorInputSubsystem
  * Subsystem for dynamically registering editor commands through scripting
  */
-UCLASS()
-class SLATESCRIPTINGCOMMANDS_API UUICommandsScriptingSubsystem : public UEngineSubsystem
+UCLASS(MinimalAPI)
+class UUICommandsScriptingSubsystem : public UEngineSubsystem
 {
 	GENERATED_BODY()
 
@@ -187,10 +189,10 @@ public:
 	 * To expose command lists from external systems to this subsystem, call FInputBindingManager::RegisterCommandList
 	 * from the external system. The binding manager will then broadcast the registered command list to all subscribers.
 	 */
-	virtual void Initialize(FSubsystemCollectionBase& Collection) override;
+	UE_API virtual void Initialize(FSubsystemCollectionBase& Collection) override;
 	
 	/** USubsystem Interface: handles deinitialization of instances of the system */
-	virtual void Deinitialize() override;
+	UE_API virtual void Deinitialize() override;
 
 	/**
 	 * Registers a command list in the given context. Optionally registering the context itself if it does not exist
@@ -201,20 +203,20 @@ public:
 	 * @param	ContextName		The context name as defined in the binding manager
 	 * @param	CommandList		The command list to register in this context
 	 */
-	void RegisterCommandListForContext(const FName ContextName, TSharedRef<FUICommandList> CommandList);
+	UE_API void RegisterCommandListForContext(const FName ContextName, TSharedRef<FUICommandList> CommandList);
 
 	/**
 	 * Unregisters a command list from the given context.
 	 * @param	ContextName		The context name as defined in the binding manager
 	 * @param	CommandList		The command list to register in this context
 	 */
-	void UnregisterCommandListForContext(const FName ContextName, TSharedRef<FUICommandList> CommandList);
+	UE_API void UnregisterCommandListForContext(const FName ContextName, TSharedRef<FUICommandList> CommandList);
 
 	/**
 	 * Unregisters all command lists from the given context.
 	 * @param	ContextName		The context name as defined in the binding manager
 	 */
-	bool UnregisterContext(const FName ContextName);
+	UE_API bool UnregisterContext(const FName ContextName);
 
 public:	
 	/**
@@ -226,7 +228,7 @@ public:
 	 * @return Whether the command was succesfully registered
 	 */
 	UFUNCTION(BlueprintCallable, Category = "Editor Scripting | Commands |")
-	bool RegisterCommand(FScriptingCommandInfo CommandInfo, FExecuteCommand OnExecuteCommand, bool bOverrideExisting = false);
+	UE_API bool RegisterCommand(FScriptingCommandInfo CommandInfo, FExecuteCommand OnExecuteCommand, bool bOverrideExisting = false);
 
 	/**
 	 * Registers a command within the given context and set.
@@ -238,7 +240,7 @@ public:
 	 * @return Whether the command was successfully registered
 	 */
 	UFUNCTION(BlueprintCallable, Category = "Editor Scripting | Commands |")
-	bool RegisterCommandChecked(FScriptingCommandInfo CommandInfo, FExecuteCommand OnExecuteCommand,
+	UE_API bool RegisterCommandChecked(FScriptingCommandInfo CommandInfo, FExecuteCommand OnExecuteCommand,
 	                            FCanExecuteCommand OnCanExecuteCommand, bool bOverrideExisting = false);
 
 	/**
@@ -246,30 +248,30 @@ public:
 	 * @return Whether the command was successfully unregistered.
 	 */
 	UFUNCTION(BlueprintCallable, Category = "Editor Scripting | Commands |")
-	bool UnregisterCommand(FScriptingCommandInfo CommandInfo);
+	UE_API bool UnregisterCommand(FScriptingCommandInfo CommandInfo);
 
 	/** Checks whether the given command is registered within the subsystem. Using name, set and context for comparison */
 	UFUNCTION(BlueprintCallable, Category = "Editor Scripting | Commands |")
-	bool IsCommandRegistered(FScriptingCommandInfo CommandInfo, bool bCheckInputChord = true) const;
+	UE_API bool IsCommandRegistered(FScriptingCommandInfo CommandInfo, bool bCheckInputChord = true) const;
 	
 	/** Retrieves the list of command info for all commands currently registered in the subsystem */
 	UFUNCTION(BlueprintCallable, Category = "Editor Scripting | Commands | Global")
-	TArray<FScriptingCommandInfo> GetRegisteredCommands() const;
+	UE_API TArray<FScriptingCommandInfo> GetRegisteredCommands() const;
 
 	/** Checks whether commands registered in the subsystem can be executed */
 	UFUNCTION(BlueprintCallable, Category = "Editor Scripting | Commands | Global")
-	bool CanExecuteCommands() const;
+	UE_API bool CanExecuteCommands() const;
 
 	/** Sets whether commands registered in the subsystem can be executed */
 	UFUNCTION(BlueprintCallable, Category = "Editor Scripting | Commands | Global")
-	void SetCanExecuteCommands(bool bShouldExecuteCommands);
+	UE_API void SetCanExecuteCommands(bool bShouldExecuteCommands);
 
 	/**
 	 * Unregisters all commands dynamically registered within all contexts and sets.
 	 * @warning this will unregister all commands currently registered by this subsystem
 	 */
 	UFUNCTION(BlueprintCallable, Category = "Editor Scripting | Commands | Global")
-	void UnregisterAllSets();
+	UE_API void UnregisterAllSets();
 
 	
 	/**
@@ -277,26 +279,26 @@ public:
 	 * @return Whether the set did not already exist and was successfully registered
 	 */
 	UFUNCTION(BlueprintCallable, Category = "Editor Scripting | Commands | Sets")
-	bool RegisterCommandSet(FName SetName);
+	UE_API bool RegisterCommandSet(FName SetName);
 
 	/** Checks whether the given set is currently registered in the subsystem */
 	UFUNCTION(BlueprintCallable, Category = "Editor Scripting | Commands | Sets")
-	bool IsCommandSetRegistered(FName SetName) const;
+	UE_API bool IsCommandSetRegistered(FName SetName) const;
 	
 	/**
 	 * Unregisters the corresponding command set with all commands registered within it
 	 * @return Whether the command set existed and was successfully unregistered
 	 */
 	UFUNCTION(BlueprintCallable, Category = "Editor Scripting | Commands | Sets")
-	bool UnregisterCommandSet(FName SetName);
+	UE_API bool UnregisterCommandSet(FName SetName);
 
 	/** Enables or disables execution of commands registered within the given set */
 	UFUNCTION(BlueprintCallable, Category = "Editor Scripting | Commands | Sets", Meta = (DisplayName = "Set Can Execute Commands"))
-	void SetCanSetExecuteCommands(FName SetName, bool bShouldExecuteCommands);
+	UE_API void SetCanSetExecuteCommands(FName SetName, bool bShouldExecuteCommands);
 
 	/** Checks whether commands in the given set can be executed. This will also check CanExecuteCommands at a global scope */
 	UFUNCTION(BlueprintCallable, Category = "Editor Scripting | Commands | Sets", Meta = (DisplayName = "Can Execute Commands"))
-	bool CanSetExecuteCommands(const FName SetName) const;
+	UE_API bool CanSetExecuteCommands(const FName SetName) const;
 
 	
 	/**
@@ -304,57 +306,57 @@ public:
 	 * This does not check whether the contexts are bound to any UI Command List. 
 	 */
 	UFUNCTION(BlueprintCallable, Category = "Editor Scripting | Commands | Contexts")
-	TArray<FName> GetAvailableContexts() const;
+	UE_API TArray<FName> GetAvailableContexts() const;
 	
 	/**
 	 * Checks whether the context with the given name is currently registered in the subsystem
 	 * This does not check whether the context is bound to any UI Command List. 
 	 */
 	UFUNCTION(BlueprintCallable, Category = "Editor Scripting | Commands | Contexts")
-	bool IsContextRegistered(FName ContextName) const;
+	UE_API bool IsContextRegistered(FName ContextName) const;
 	
 	/**
 	 * Retrieves the number of UI Command Lists registered within this context through the subsystem.
 	 * UI Command Lists are typically used to bind the list of commands associated with a single UI (i.e. a single viewport).
 	 */
 	UFUNCTION(BlueprintCallable, Category = "Editor Scripting | Commands | Contexts")
-	int GetBindingCountForContext(FName ContextName);
+	UE_API int GetBindingCountForContext(FName ContextName);
 
 	/**
 	 * Checks whether the given input chord is already mapped to a command in the given context.
 	 * This includes commands not registered through the subsystem.
 	 */
 	UFUNCTION(BlueprintCallable, Category = "Editor Scripting | Commands | Contexts")
-	bool IsInputChordMapped(FName ContextName, FInputChord InputChord) const;	
+	UE_API bool IsInputChordMapped(FName ContextName, FInputChord InputChord) const;	
 	
 private:
 	/** The internal method to dynamically registers an editor command */
-	bool RegisterNewScriptingCommand(FScriptingCommandInfo CommandInfo, FExecuteCommand OnExecuteCommand,
+	UE_API bool RegisterNewScriptingCommand(FScriptingCommandInfo CommandInfo, FExecuteCommand OnExecuteCommand,
 	                               FCanExecuteAction OnCanExecuteAction, bool bOverrideExisting = false);
 
 	/** Registers a context command in its command lists. This assumes the command is already contained in a set */
-	bool RegisterScriptingCommand(const TSharedRef<FScriptingCommand> ScriptingCommand);
+	UE_API bool RegisterScriptingCommand(const TSharedRef<FScriptingCommand> ScriptingCommand);
 
 	/** Unregisters a context command from its command lists. This assumes the command is already contained in a set */
-	bool UnregisterScriptingCommand(const TSharedRef<FScriptingCommand> ScriptingCommand);
+	UE_API bool UnregisterScriptingCommand(const TSharedRef<FScriptingCommand> ScriptingCommand);
 
 	/** The delegate bound to registered commands. It will broadcast back to the given dynamic delegate */
 	UFUNCTION()
-	static void HandleExecuteAction(FExecuteCommand OnExecuteAction, const FScriptingCommandInfo CommandInfo);
+	static UE_API void HandleExecuteAction(FExecuteCommand OnExecuteAction, const FScriptingCommandInfo CommandInfo);
 
 	/**
 	 * The delegate used to check whether a command can be executed.
 	 * @return true if the given delegate returns true and the given set and all subsystem commands are enabled
 	 */
 	UFUNCTION()
-	bool HandleCanExecuteAction(FCanExecuteCommand OnCanExecuteAction, const FScriptingCommandInfo CommandInfo) const;
+	UE_API bool HandleCanExecuteAction(FCanExecuteCommand OnCanExecuteAction, const FScriptingCommandInfo CommandInfo) const;
 
 	/**
 	 * The default delegate to check if a command can be executed (when the user does not provide a specific delegate).
 	 * @return true if the given set and all subsystem commands are enabled
 	 */
 	UFUNCTION()
-	bool DefaultCanExecuteAction(FName SetName) const;
+	UE_API bool DefaultCanExecuteAction(FName SetName) const;
 
 private:
 	/** The list of commands contexts. Each context contains its own list of registered commands and UI command lists */
@@ -365,4 +367,6 @@ private:
 
 	/** Whether the commands for all the sets of the subsystem can be executed */
 	bool bCanExecuteCommands = true;
-};
\ No newline at end of file
+};
+
+#undef UE_API
diff --git a/Engine/Plugins/Tests/InterchangeTests/Source/InterchangeTests/Public/ImportTestFunctions/ActorImportTestFunctions.h b/Engine/Plugins/Tests/InterchangeTests/Source/InterchangeTests/Public/ImportTestFunctions/ActorImportTestFunctions.h
--- a/Engine/Plugins/Tests/InterchangeTests/Source/InterchangeTests/Public/ImportTestFunctions/ActorImportTestFunctions.h
+++ b/Engine/Plugins/Tests/InterchangeTests/Source/InterchangeTests/Public/ImportTestFunctions/ActorImportTestFunctions.h
@@ -5,37 +5,41 @@
 #include "ImportTestFunctionsBase.h"
 #include "ActorImportTestFunctions.generated.h"
 
+#define UE_API INTERCHANGETESTS_API
+
 class AActor;
 struct FInterchangeTestFunctionResult;
 
 
-UCLASS()
-class INTERCHANGETESTS_API UActorImportTestFunctions : public UImportTestFunctionsBase
+UCLASS(MinimalAPI)
+class UActorImportTestFunctions : public UImportTestFunctionsBase
 {
 	GENERATED_BODY()
 
 public:
 
 	// UImportTestFunctionsBase interface
-	virtual UClass* GetAssociatedAssetType() const override;
+	UE_API virtual UClass* GetAssociatedAssetType() const override;
 
 	/** Check whether the expected number of actors are imported */
 	UFUNCTION(Exec)
-	static FInterchangeTestFunctionResult CheckImportedActorCount(const TArray<AActor*>& Actors, int32 ExpectedNumberOfImportedActors);
+	static UE_API FInterchangeTestFunctionResult CheckImportedActorCount(const TArray<AActor*>& Actors, int32 ExpectedNumberOfImportedActors);
 
 	/** Check whether the expected number of actors with a given class are imported */
 	UFUNCTION(Exec)
-	static FInterchangeTestFunctionResult CheckActorClassCount(const TArray<AActor*>& Actors, TSubclassOf<AActor> Class, int32 ExpectedNumberOfActors);
+	static UE_API FInterchangeTestFunctionResult CheckActorClassCount(const TArray<AActor*>& Actors, TSubclassOf<AActor> Class, int32 ExpectedNumberOfActors);
 
 	/** Check whether the imported actor has the expected class */
 	UFUNCTION(Exec)
-	static FInterchangeTestFunctionResult CheckActorClass(const AActor* Actor, TSubclassOf<AActor> ExpectedClass);
+	static UE_API FInterchangeTestFunctionResult CheckActorClass(const AActor* Actor, TSubclassOf<AActor> ExpectedClass);
 
 	/** Check whether the generic property (with a given name) in the imported actor has the expected value (or matches with it using regex) */
 	UFUNCTION(Exec)
-	static FInterchangeTestFunctionResult CheckActorPropertyValue(const AActor* Actor, const FString& PropertyName, bool bUseRegexToMatchValue, const FString& ExpectedValue);
+	static UE_API FInterchangeTestFunctionResult CheckActorPropertyValue(const AActor* Actor, const FString& PropertyName, bool bUseRegexToMatchValue, const FString& ExpectedValue);
 
 	/** Check whether the generic property (with a given name) in the imported actor component (with the given name) has the expected value (or matches with it using regex) */
 	UFUNCTION(Exec)
-	static FInterchangeTestFunctionResult CheckComponentPropertyValue(const AActor* Actor, const FString& ComponentName, const FString& PropertyName, bool bUseRegexToMatchValue, const FString& ExpectedValue);
+	static UE_API FInterchangeTestFunctionResult CheckComponentPropertyValue(const AActor* Actor, const FString& ComponentName, const FString& PropertyName, bool bUseRegexToMatchValue, const FString& ExpectedValue);
 };
+
+#undef UE_API
diff --git a/Engine/Plugins/Tests/InterchangeTests/Source/InterchangeTests/Public/ImportTestFunctions/AnimationImportTestFunctions.h b/Engine/Plugins/Tests/InterchangeTests/Source/InterchangeTests/Public/ImportTestFunctions/AnimationImportTestFunctions.h
--- a/Engine/Plugins/Tests/InterchangeTests/Source/InterchangeTests/Public/ImportTestFunctions/AnimationImportTestFunctions.h
+++ b/Engine/Plugins/Tests/InterchangeTests/Source/InterchangeTests/Public/ImportTestFunctions/AnimationImportTestFunctions.h
@@ -5,52 +5,56 @@
 #include "ImportTestFunctionsBase.h"
 #include "AnimationImportTestFunctions.generated.h"
 
+#define UE_API INTERCHANGETESTS_API
+
 struct FInterchangeTestFunctionResult;
 
 
-UCLASS()
-class INTERCHANGETESTS_API UAnimationImportTestFunctions : public UImportTestFunctionsBase
+UCLASS(MinimalAPI)
+class UAnimationImportTestFunctions : public UImportTestFunctionsBase
 {
 	GENERATED_BODY()
 
 public:
 
 	// UImportTestFunctionsBase interface
-	virtual UClass* GetAssociatedAssetType() const override;
+	UE_API virtual UClass* GetAssociatedAssetType() const override;
 
 	/** Check whether the expected number of anim sequences are imported */
 	UFUNCTION(Exec)
-	static FInterchangeTestFunctionResult CheckImportedAnimSequenceCount(const TArray<UAnimSequence*>& AnimSequences, int32 ExpectedNumberOfImportedAnimSequences);
+	static UE_API FInterchangeTestFunctionResult CheckImportedAnimSequenceCount(const TArray<UAnimSequence*>& AnimSequences, int32 ExpectedNumberOfImportedAnimSequences);
 
 	/** Check whether the animation length (second) is as expected */
 	UFUNCTION(Exec)
-	static FInterchangeTestFunctionResult CheckAnimationLength(UAnimSequence* AnimSequence, float ExpectedAnimationLength);
+	static UE_API FInterchangeTestFunctionResult CheckAnimationLength(UAnimSequence* AnimSequence, float ExpectedAnimationLength);
 
 	/** Check whether the animation frame number is as expected */
 	UFUNCTION(Exec)
-	static FInterchangeTestFunctionResult CheckAnimationFrameNumber(UAnimSequence* AnimSequence, int32 ExpectedFrameNumber);
+	static UE_API FInterchangeTestFunctionResult CheckAnimationFrameNumber(UAnimSequence* AnimSequence, int32 ExpectedFrameNumber);
 
 	/** Check whether the given curve key time(sec) for the given curve name has the expected time(sec) */
 	UFUNCTION(Exec)
-	static FInterchangeTestFunctionResult CheckCurveKeyTime(UAnimSequence* AnimSequence, const FString& CurveName, int32 KeyIndex, float ExpectedCurveKeyTime);
+	static UE_API FInterchangeTestFunctionResult CheckCurveKeyTime(UAnimSequence* AnimSequence, const FString& CurveName, int32 KeyIndex, float ExpectedCurveKeyTime);
 
 	/** Check whether the given curve key value for the given curve name has the expected value */
 	UFUNCTION(Exec)
-	static FInterchangeTestFunctionResult CheckCurveKeyValue(UAnimSequence* AnimSequence, const FString& CurveName, int32 KeyIndex, float ExpectedCurveKeyValue);
+	static UE_API FInterchangeTestFunctionResult CheckCurveKeyValue(UAnimSequence* AnimSequence, const FString& CurveName, int32 KeyIndex, float ExpectedCurveKeyValue);
 
 	/** Check whether the given curve key arrive tangent for the given curve name has the expected arrive tangent */
 	UFUNCTION(Exec)
-	static FInterchangeTestFunctionResult CheckCurveKeyArriveTangent(UAnimSequence* AnimSequence, const FString& CurveName, int32 KeyIndex, float ExpectedCurveKeyArriveTangent);
+	static UE_API FInterchangeTestFunctionResult CheckCurveKeyArriveTangent(UAnimSequence* AnimSequence, const FString& CurveName, int32 KeyIndex, float ExpectedCurveKeyArriveTangent);
 
 	/** Check whether the given curve key arrive tangent weight for the given curve name has the expected arrive tangent weight */
 	UFUNCTION(Exec)
-	static FInterchangeTestFunctionResult CheckCurveKeyArriveTangentWeight(UAnimSequence* AnimSequence, const FString& CurveName, int32 KeyIndex, float ExpectedCurveKeyArriveTangentWeight);
+	static UE_API FInterchangeTestFunctionResult CheckCurveKeyArriveTangentWeight(UAnimSequence* AnimSequence, const FString& CurveName, int32 KeyIndex, float ExpectedCurveKeyArriveTangentWeight);
 
 	/** Check whether the given curve key leave tangent for the given curve name has the expected leave tangent */
 	UFUNCTION(Exec)
-	static FInterchangeTestFunctionResult CheckCurveKeyLeaveTangent(UAnimSequence* AnimSequence, const FString& CurveName, int32 KeyIndex, float ExpectedCurveKeyLeaveTangent);
+	static UE_API FInterchangeTestFunctionResult CheckCurveKeyLeaveTangent(UAnimSequence* AnimSequence, const FString& CurveName, int32 KeyIndex, float ExpectedCurveKeyLeaveTangent);
 
 	/** Check whether the given curve key leave tangent weight for the given curve name has the expected leave tangent weight */
 	UFUNCTION(Exec)
-	static FInterchangeTestFunctionResult CheckCurveKeyLeaveTangentWeight(UAnimSequence* AnimSequence, const FString& CurveName, int32 KeyIndex, float ExpectedCurveKeyLeaveTangentWeight);
+	static UE_API FInterchangeTestFunctionResult CheckCurveKeyLeaveTangentWeight(UAnimSequence* AnimSequence, const FString& CurveName, int32 KeyIndex, float ExpectedCurveKeyLeaveTangentWeight);
 };
+
+#undef UE_API
diff --git a/Engine/Plugins/Tests/InterchangeTests/Source/InterchangeTests/Public/ImportTestFunctions/AssetImportTestFunctions.h b/Engine/Plugins/Tests/InterchangeTests/Source/InterchangeTests/Public/ImportTestFunctions/AssetImportTestFunctions.h
--- a/Engine/Plugins/Tests/InterchangeTests/Source/InterchangeTests/Public/ImportTestFunctions/AssetImportTestFunctions.h
+++ b/Engine/Plugins/Tests/InterchangeTests/Source/InterchangeTests/Public/ImportTestFunctions/AssetImportTestFunctions.h
@@ -5,31 +5,33 @@
 #include "ImportTestFunctionsBase.h"
 #include "AssetImportTestFunctions.generated.h"
 
+#define UE_API INTERCHANGETESTS_API
+
 struct FInterchangeTestFunctionResult;
 class UObject;
 
 
-UCLASS()
-class INTERCHANGETESTS_API UAssetImportTestFunctions : public UImportTestFunctionsBase
+UCLASS(MinimalAPI)
+class UAssetImportTestFunctions : public UImportTestFunctionsBase
 {
 	GENERATED_BODY()
 
 public:
 
 	// UImportTestFunctionsBase interface
-	virtual UClass* GetAssociatedAssetType() const override;
+	UE_API virtual UClass* GetAssociatedAssetType() const override;
 
 
 	/** Check whether the expected number of metadata for the object are imported */
 	UFUNCTION(Exec)
-	static FInterchangeTestFunctionResult CheckImportedMetadataCount(const UObject* Object, const int32 ExpectedNumberOfMetadataForThisObject);
+	static UE_API FInterchangeTestFunctionResult CheckImportedMetadataCount(const UObject* Object, const int32 ExpectedNumberOfMetadataForThisObject);
 
 	/** 
 	 * Check whether the expected object metadata key exist.
 	 * @Param ExpectedMetadataKey - The object metadata key to pass to the package to retrieve the metadata value
 	 */
 	UFUNCTION(Exec)
-	static FInterchangeTestFunctionResult CheckMetadataExist(const UObject* Object, const FString& ExpectedMetadataKey);
+	static UE_API FInterchangeTestFunctionResult CheckMetadataExist(const UObject* Object, const FString& ExpectedMetadataKey);
 
 	/**
 	 * Check whether the expected object metadata value is imported.
@@ -37,13 +39,15 @@ public:
 	 * @Param ExpectedMetadataValue - The value to compare the object metadata query with the metadata key
 	 */
 	UFUNCTION(Exec)
-	static FInterchangeTestFunctionResult CheckMetadataValue(const UObject* Object, const FString& ExpectedMetadataKey, const FString& ExpectedMetadataValue);
+	static UE_API FInterchangeTestFunctionResult CheckMetadataValue(const UObject* Object, const FString& ExpectedMetadataKey, const FString& ExpectedMetadataValue);
 
 	/**
 	 * Check whether the object's path has the given string as substring.
 	 * @Param ExpectedPathString - The string to look for in the object path.
 	 */
 	UFUNCTION(Exec)
-	static FInterchangeTestFunctionResult CheckObjectPathHasSubstring(const UObject* Object, const FString& ExpectedPathString);
+	static UE_API FInterchangeTestFunctionResult CheckObjectPathHasSubstring(const UObject* Object, const FString& ExpectedPathString);
 
 };
+
+#undef UE_API
diff --git a/Engine/Plugins/Tests/InterchangeTests/Source/InterchangeTests/Public/ImportTestFunctions/ImportTestFunctionsBase.h b/Engine/Plugins/Tests/InterchangeTests/Source/InterchangeTests/Public/ImportTestFunctions/ImportTestFunctionsBase.h
--- a/Engine/Plugins/Tests/InterchangeTests/Source/InterchangeTests/Public/ImportTestFunctions/ImportTestFunctionsBase.h
+++ b/Engine/Plugins/Tests/InterchangeTests/Source/InterchangeTests/Public/ImportTestFunctions/ImportTestFunctionsBase.h
@@ -4,6 +4,8 @@
 
 #include "ImportTestFunctionsBase.generated.h"
 
+#define UE_API INTERCHANGETESTS_API
+
 
 /**
  * This is the base class for any static class which provides test functions for an asset type.
@@ -12,16 +14,18 @@
  * However these functions are in reality not designed to be called from the console.
  * @todo: create an alternative UFUNCTION tag for exporting default parameters as metadata.
  */
-UCLASS()
-class INTERCHANGETESTS_API UImportTestFunctionsBase : public UObject
+UCLASS(MinimalAPI)
+class UImportTestFunctionsBase : public UObject
 {
 	GENERATED_BODY()
 
 public:
 
 	/** Determine which of the import test functions classes to use for a given asset class type */
-	static UClass* GetImportTestFunctionsClassForAssetType(UClass* AssetClass);
+	static UE_API UClass* GetImportTestFunctionsClassForAssetType(UClass* AssetClass);
 
 	// UImportTestFunctionsBase interface
 	virtual UClass* GetAssociatedAssetType() const { return UObject::StaticClass(); }
 };
+
+#undef UE_API
diff --git a/Engine/Plugins/Tests/InterchangeTests/Source/InterchangeTests/Public/ImportTestFunctions/InterchangeResultImportTestFunctions.h b/Engine/Plugins/Tests/InterchangeTests/Source/InterchangeTests/Public/ImportTestFunctions/InterchangeResultImportTestFunctions.h
--- a/Engine/Plugins/Tests/InterchangeTests/Source/InterchangeTests/Public/ImportTestFunctions/InterchangeResultImportTestFunctions.h
+++ b/Engine/Plugins/Tests/InterchangeTests/Source/InterchangeTests/Public/ImportTestFunctions/InterchangeResultImportTestFunctions.h
@@ -6,23 +6,27 @@
 #include "UObject/Package.h"
 #include "InterchangeResultImportTestFunctions.generated.h"
 
+#define UE_API INTERCHANGETESTS_API
+
 class UInterchangeResult;
 class UInterchangeResultsContainer;
 
 struct FInterchangeTestFunctionResult;
 
 
-UCLASS()
-class INTERCHANGETESTS_API UInterchangeResultImportTestFunctions : public UImportTestFunctionsBase
+UCLASS(MinimalAPI)
+class UInterchangeResultImportTestFunctions : public UImportTestFunctionsBase
 {
 	GENERATED_BODY()
 
 public:
 
 	// UImportTestFunctionsBase interface
-	virtual UClass* GetAssociatedAssetType() const override;
+	UE_API virtual UClass* GetAssociatedAssetType() const override;
 
 	/** Check whether the specified InterchangeResult was emitted during import */
 	UFUNCTION(Exec)
-	static FInterchangeTestFunctionResult CheckIfErrorOrWarningWasGenerated(UInterchangeResultsContainer* ResultsContainer, TSubclassOf<UInterchangeResult> ErrorOrWarningClass);
+	static UE_API FInterchangeTestFunctionResult CheckIfErrorOrWarningWasGenerated(UInterchangeResultsContainer* ResultsContainer, TSubclassOf<UInterchangeResult> ErrorOrWarningClass);
 };
+
+#undef UE_API
diff --git a/Engine/Plugins/Tests/InterchangeTests/Source/InterchangeTests/Public/ImportTestFunctions/LevelSequenceImportTestFunctions.h b/Engine/Plugins/Tests/InterchangeTests/Source/InterchangeTests/Public/ImportTestFunctions/LevelSequenceImportTestFunctions.h
--- a/Engine/Plugins/Tests/InterchangeTests/Source/InterchangeTests/Public/ImportTestFunctions/LevelSequenceImportTestFunctions.h
+++ b/Engine/Plugins/Tests/InterchangeTests/Source/InterchangeTests/Public/ImportTestFunctions/LevelSequenceImportTestFunctions.h
@@ -6,33 +6,37 @@
 #include "ImportTestFunctionsBase.h"
 #include "LevelSequenceImportTestFunctions.generated.h"
 
+#define UE_API INTERCHANGETESTS_API
+
 class ULevelSequence;
 struct FInterchangeTestFunctionResult;
 
 
-UCLASS()
-class INTERCHANGETESTS_API ULevelSequenceImportTestFunctions : public UImportTestFunctionsBase
+UCLASS(MinimalAPI)
+class ULevelSequenceImportTestFunctions : public UImportTestFunctionsBase
 {
 	GENERATED_BODY()
 
 public:
 
 	// UImportTestFunctionsBase interface
-	virtual UClass* GetAssociatedAssetType() const override;
+	UE_API virtual UClass* GetAssociatedAssetType() const override;
 
 	/** Check whether the expected number of level sequences are imported */
 	UFUNCTION(Exec)
-	static FInterchangeTestFunctionResult CheckLevelSequenceCount(const TArray<ULevelSequence*>& LevelSequences, int32 ExpectedNumberOfLevelSequences);
+	static UE_API FInterchangeTestFunctionResult CheckLevelSequenceCount(const TArray<ULevelSequence*>& LevelSequences, int32 ExpectedNumberOfLevelSequences);
 
 	/** Check whether the level sequence length (second) is as expected */
 	UFUNCTION(Exec)
-	static FInterchangeTestFunctionResult CheckSequenceLength(const ULevelSequence* LevelSequence, float ExpectedLevelSequenceLength);
+	static UE_API FInterchangeTestFunctionResult CheckSequenceLength(const ULevelSequence* LevelSequence, float ExpectedLevelSequenceLength);
 
 	/** Check whether the imported level sequence has the expected number of sections */
 	UFUNCTION(Exec)
-	static FInterchangeTestFunctionResult CheckSectionCount(const ULevelSequence* LevelSequence, int32 ExpectedNumberOfSections);
+	static UE_API FInterchangeTestFunctionResult CheckSectionCount(const ULevelSequence* LevelSequence, int32 ExpectedNumberOfSections);
 
 	/** Check whether the imported level sequence has the expected interpolation mode for the given section */
 	UFUNCTION(Exec)
-	static FInterchangeTestFunctionResult CheckSectionInterpolationMode(const ULevelSequence* LevelSequence, int32 SectionIndex, ERichCurveInterpMode ExpectedInterpolationMode);
+	static UE_API FInterchangeTestFunctionResult CheckSectionInterpolationMode(const ULevelSequence* LevelSequence, int32 SectionIndex, ERichCurveInterpMode ExpectedInterpolationMode);
 };
+
+#undef UE_API
diff --git a/Engine/Plugins/Tests/InterchangeTests/Source/InterchangeTests/Public/ImportTestFunctions/LevelVariantSetsImportTestFunctions.h b/Engine/Plugins/Tests/InterchangeTests/Source/InterchangeTests/Public/ImportTestFunctions/LevelVariantSetsImportTestFunctions.h
--- a/Engine/Plugins/Tests/InterchangeTests/Source/InterchangeTests/Public/ImportTestFunctions/LevelVariantSetsImportTestFunctions.h
+++ b/Engine/Plugins/Tests/InterchangeTests/Source/InterchangeTests/Public/ImportTestFunctions/LevelVariantSetsImportTestFunctions.h
@@ -5,34 +5,38 @@
 #include "ImportTestFunctionsBase.h"
 #include "LevelVariantSetsImportTestFunctions.generated.h"
 
+#define UE_API INTERCHANGETESTS_API
+
 class ULevelVariantSets;
 
 struct FInterchangeTestFunctionResult;
 
 
-UCLASS()
-class INTERCHANGETESTS_API ULevelVariantSetsImportTestFunctions : public UImportTestFunctionsBase
+UCLASS(MinimalAPI)
+class ULevelVariantSetsImportTestFunctions : public UImportTestFunctionsBase
 {
 	GENERATED_BODY()
 
 public:
 
 	// UImportTestFunctionsBase interface
-	virtual UClass* GetAssociatedAssetType() const override;
+	UE_API virtual UClass* GetAssociatedAssetType() const override;
 
 	/** Check whether the expected number of LevelVariantSets are imported */
 	UFUNCTION(Exec)
-	static FInterchangeTestFunctionResult CheckLevelVariantSetsCount(const TArray<ULevelVariantSets*>& LevelVariantSetsAssets, int32 ExpectedNumberOfLevelVariantSets);
+	static UE_API FInterchangeTestFunctionResult CheckLevelVariantSetsCount(const TArray<ULevelVariantSets*>& LevelVariantSetsAssets, int32 ExpectedNumberOfLevelVariantSets);
 
 	/** Check whether the imported LevelVariantSets has the expected number of variant sets */
 	UFUNCTION(Exec)
-	static FInterchangeTestFunctionResult CheckVariantSetsCount(ULevelVariantSets* LevelVariantSets, int32 ExpectedNumberOfVariantSets);
+	static UE_API FInterchangeTestFunctionResult CheckVariantSetsCount(ULevelVariantSets* LevelVariantSets, int32 ExpectedNumberOfVariantSets);
 
 	/** Check whether the imported LevelVariantSets has the expected number of variants for the given variant set */
 	UFUNCTION(Exec)
-	static FInterchangeTestFunctionResult CheckVariantsCount(ULevelVariantSets* LevelVariantSets, const FString& VariantSetName, int32 ExpectedNumberOfVariants);
+	static UE_API FInterchangeTestFunctionResult CheckVariantsCount(ULevelVariantSets* LevelVariantSets, const FString& VariantSetName, int32 ExpectedNumberOfVariants);
 
 	/** Check whether the imported LevelVariantSets has the expected number of bindings for the given variant in the given set */
 	UFUNCTION(Exec)
-	static FInterchangeTestFunctionResult CheckBindingsCount(ULevelVariantSets* LevelVariantSets, const FString& VariantSetName, const FString& VariantName, int32 ExpectedNumberOfBindings);
+	static UE_API FInterchangeTestFunctionResult CheckBindingsCount(ULevelVariantSets* LevelVariantSets, const FString& VariantSetName, const FString& VariantName, int32 ExpectedNumberOfBindings);
 };
+
+#undef UE_API
diff --git a/Engine/Plugins/Tests/InterchangeTests/Source/InterchangeTests/Public/ImportTestFunctions/LightImportTestFunctions.h b/Engine/Plugins/Tests/InterchangeTests/Source/InterchangeTests/Public/ImportTestFunctions/LightImportTestFunctions.h
--- a/Engine/Plugins/Tests/InterchangeTests/Source/InterchangeTests/Public/ImportTestFunctions/LightImportTestFunctions.h
+++ b/Engine/Plugins/Tests/InterchangeTests/Source/InterchangeTests/Public/ImportTestFunctions/LightImportTestFunctions.h
@@ -6,32 +6,36 @@
 #include "InterchangeTestFunction.h"
 #include "LightImportTestFunctions.generated.h"
 
+#define UE_API INTERCHANGETESTS_API
+
 struct FInterchangeTestFunctionResult;
 
 
-UCLASS()
-class INTERCHANGETESTS_API ULightImportTestFunctions : public UActorImportTestFunctions
+UCLASS(MinimalAPI)
+class ULightImportTestFunctions : public UActorImportTestFunctions
 {
 	GENERATED_BODY()
 
 public:
 
 	// UImportTestFunctionsBase interface
-	virtual UClass* GetAssociatedAssetType() const override;
+	UE_API virtual UClass* GetAssociatedAssetType() const override;
 
 	/** Check whether the light position is correct*/
 	UFUNCTION(Exec)
-	static FInterchangeTestFunctionResult CheckLightPosition(ALight* Light, const FVector& ExpectedLightPosition);
+	static UE_API FInterchangeTestFunctionResult CheckLightPosition(ALight* Light, const FVector& ExpectedLightPosition);
 
 	/** Check whether the light direction is correct*/
 	UFUNCTION(Exec)
-	static FInterchangeTestFunctionResult CheckLightDirection(ALight* Light, const FVector& ExpectedLightDirection);
+	static UE_API FInterchangeTestFunctionResult CheckLightDirection(ALight* Light, const FVector& ExpectedLightDirection);
 
 	/** Check whether the light intensity is correct*/
 	UFUNCTION(Exec)
-	static FInterchangeTestFunctionResult CheckLightIntensity(ALight* Light, float ExpectedLightIntensity);
+	static UE_API FInterchangeTestFunctionResult CheckLightIntensity(ALight* Light, float ExpectedLightIntensity);
 
 	/** Check whether the light color is correct*/
 	UFUNCTION(Exec)
-	static FInterchangeTestFunctionResult CheckLightColor(ALight* Light, const FLinearColor& ExpectedLightColor);
+	static UE_API FInterchangeTestFunctionResult CheckLightColor(ALight* Light, const FLinearColor& ExpectedLightColor);
 };
+
+#undef UE_API
diff --git a/Engine/Plugins/Tests/InterchangeTests/Source/InterchangeTests/Public/ImportTestFunctions/MaterialImportTestFunctions.h b/Engine/Plugins/Tests/InterchangeTests/Source/InterchangeTests/Public/ImportTestFunctions/MaterialImportTestFunctions.h
--- a/Engine/Plugins/Tests/InterchangeTests/Source/InterchangeTests/Public/ImportTestFunctions/MaterialImportTestFunctions.h
+++ b/Engine/Plugins/Tests/InterchangeTests/Source/InterchangeTests/Public/ImportTestFunctions/MaterialImportTestFunctions.h
@@ -6,50 +6,54 @@
 #include "ImportTestFunctionsBase.h"
 #include "MaterialImportTestFunctions.generated.h"
 
+#define UE_API INTERCHANGETESTS_API
+
 class UMaterialInterface;
 
 struct FInterchangeTestFunctionResult;
 
 
-UCLASS()
-class INTERCHANGETESTS_API UMaterialImportTestFunctions : public UImportTestFunctionsBase
+UCLASS(MinimalAPI)
+class UMaterialImportTestFunctions : public UImportTestFunctionsBase
 {
 	GENERATED_BODY()
 
 public:
 
 	// UImportTestFunctionsBase interface
-	virtual UClass* GetAssociatedAssetType() const override;
+	UE_API virtual UClass* GetAssociatedAssetType() const override;
 
 	/** Check whether the expected number of materials are imported */
 	UFUNCTION(Exec)
-	static FInterchangeTestFunctionResult CheckImportedMaterialCount(const TArray<UMaterialInterface*>& MaterialInterfaces, int32 ExpectedNumberOfImportedMaterials);
+	static UE_API FInterchangeTestFunctionResult CheckImportedMaterialCount(const TArray<UMaterialInterface*>& MaterialInterfaces, int32 ExpectedNumberOfImportedMaterials);
 
 	/** Check whether the expected number of material instances are imported */
 	UFUNCTION(Exec)
-	static FInterchangeTestFunctionResult CheckImportedMaterialInstanceCount(const TArray<UMaterialInterface*>& MaterialInterfaces, int32 ExpectedNumberOfImportedMaterialInstances);
+	static UE_API FInterchangeTestFunctionResult CheckImportedMaterialInstanceCount(const TArray<UMaterialInterface*>& MaterialInterfaces, int32 ExpectedNumberOfImportedMaterialInstances);
 
 	/** Check whether the imported material has the expected shading model */
 	UFUNCTION(Exec)
-	static FInterchangeTestFunctionResult CheckShadingModel(const UMaterialInterface* MaterialInterface, EMaterialShadingModel ExpectedShadingModel);
+	static UE_API FInterchangeTestFunctionResult CheckShadingModel(const UMaterialInterface* MaterialInterface, EMaterialShadingModel ExpectedShadingModel);
 
 	/** Check whether the imported material has the expected blend mode */
 	UFUNCTION(Exec)
-	static FInterchangeTestFunctionResult CheckBlendMode(const UMaterialInterface* MaterialInterface, EBlendMode ExpectedBlendMode);
+	static UE_API FInterchangeTestFunctionResult CheckBlendMode(const UMaterialInterface* MaterialInterface, EBlendMode ExpectedBlendMode);
 
 	/** Check whether the imported material has the expected two-sided setting */
 	UFUNCTION(Exec)
-	static FInterchangeTestFunctionResult CheckIsTwoSided(const UMaterialInterface* MaterialInterface, bool ExpectedIsTwoSided);
+	static UE_API FInterchangeTestFunctionResult CheckIsTwoSided(const UMaterialInterface* MaterialInterface, bool ExpectedIsTwoSided);
 
 	/** Check whether the imported material has the expected opacity mask clip value */
 	UFUNCTION(Exec)
-	static FInterchangeTestFunctionResult CheckOpacityMaskClipValue(const UMaterialInterface* MaterialInterface, float ExpectedOpacityMaskClipValue);
+	static UE_API FInterchangeTestFunctionResult CheckOpacityMaskClipValue(const UMaterialInterface* MaterialInterface, float ExpectedOpacityMaskClipValue);
 
 	/** Check whether the imported material has the expected scalar parameter value */
 	UFUNCTION(Exec)
-	static FInterchangeTestFunctionResult CheckScalarParameter(const UMaterialInterface* MaterialInterface, const FString& ParameterName, float ExpectedParameterValue);
+	static UE_API FInterchangeTestFunctionResult CheckScalarParameter(const UMaterialInterface* MaterialInterface, const FString& ParameterName, float ExpectedParameterValue);
 
 	/** Check whether the imported material has the expected vector parameter value */
 	UFUNCTION(Exec)
-	static FInterchangeTestFunctionResult CheckVectorParameter(const UMaterialInterface* MaterialInterface, const FString& ParameterName, FLinearColor ExpectedParameterValue);
+	static UE_API FInterchangeTestFunctionResult CheckVectorParameter(const UMaterialInterface* MaterialInterface, const FString& ParameterName, FLinearColor ExpectedParameterValue);
 };
+
+#undef UE_API
diff --git a/Engine/Plugins/Tests/InterchangeTests/Source/InterchangeTests/Public/ImportTestFunctions/MaterialXTestFunctions.h b/Engine/Plugins/Tests/InterchangeTests/Source/InterchangeTests/Public/ImportTestFunctions/MaterialXTestFunctions.h
--- a/Engine/Plugins/Tests/InterchangeTests/Source/InterchangeTests/Public/ImportTestFunctions/MaterialXTestFunctions.h
+++ b/Engine/Plugins/Tests/InterchangeTests/Source/InterchangeTests/Public/ImportTestFunctions/MaterialXTestFunctions.h
@@ -6,27 +6,31 @@
 
 #include "MaterialXTestFunctions.generated.h"
 
+#define UE_API INTERCHANGETESTS_API
+
 class UImportTestFunctionsBase;
 
 struct FInterchangeTestFunctionResult;
 
 class UMaterialInterface;
 
-UCLASS()
-class INTERCHANGETESTS_API UMaterialXTestFunctions : public UImportTestFunctionsBase
+UCLASS(MinimalAPI)
+class UMaterialXTestFunctions : public UImportTestFunctionsBase
 {
 	GENERATED_BODY()
 
 public:
 
 	// UImportTestFunctionsBase interface
-	virtual UClass* GetAssociatedAssetType() const override;
+	UE_API virtual UClass* GetAssociatedAssetType() const override;
 
 	/** Check whether the expected number of inputs are connected to the MX_StandardSurface material function */
 	UFUNCTION(Exec, meta = (DisplayName = "MX: Check Connected Input Count"))
-	static FInterchangeTestFunctionResult CheckConnectedInputCount(const UMaterialInterface* MaterialInterface, int32 ExpectedNumber);
+	static UE_API FInterchangeTestFunctionResult CheckConnectedInputCount(const UMaterialInterface* MaterialInterface, int32 ExpectedNumber);
 
 	/** Check whether a specific input of the MX_StandardSurface material function is connected or not */
 	UFUNCTION(Exec, meta = (DisplayName = "MX: Check Input Is Connected"))
-	static FInterchangeTestFunctionResult CheckInputConnected(const UMaterialInterface* MaterialInterface, const FString& InputName, bool bIsConnected);
+	static UE_API FInterchangeTestFunctionResult CheckInputConnected(const UMaterialInterface* MaterialInterface, const FString& InputName, bool bIsConnected);
 };
+
+#undef UE_API
diff --git a/Engine/Plugins/Tests/InterchangeTests/Source/InterchangeTests/Public/ImportTestFunctions/PointLightImportTestFunctions.h b/Engine/Plugins/Tests/InterchangeTests/Source/InterchangeTests/Public/ImportTestFunctions/PointLightImportTestFunctions.h
--- a/Engine/Plugins/Tests/InterchangeTests/Source/InterchangeTests/Public/ImportTestFunctions/PointLightImportTestFunctions.h
+++ b/Engine/Plugins/Tests/InterchangeTests/Source/InterchangeTests/Public/ImportTestFunctions/PointLightImportTestFunctions.h
@@ -6,20 +6,24 @@
 #include "InterchangeTestFunction.h"
 #include "PointLightImportTestFunctions.generated.h"
 
+#define UE_API INTERCHANGETESTS_API
+
 struct FInterchangeTestFunctionResult;
 
 
-UCLASS()
-class INTERCHANGETESTS_API UPointLightImportTestFunctions : public ULightImportTestFunctions
+UCLASS(MinimalAPI)
+class UPointLightImportTestFunctions : public ULightImportTestFunctions
 {
 	GENERATED_BODY()
 
 public:
 
 	// UImportTestFunctionsBase interface
-	virtual UClass* GetAssociatedAssetType() const override;
+	UE_API virtual UClass* GetAssociatedAssetType() const override;
 
 	/** Check whether the light falloff exponent is correct*/
 	UFUNCTION(Exec)
-	static FInterchangeTestFunctionResult CheckLightFalloffExponent(APointLight* Light, float ExpectedLightFalloff);
+	static UE_API FInterchangeTestFunctionResult CheckLightFalloffExponent(APointLight* Light, float ExpectedLightFalloff);
 };
+
+#undef UE_API
diff --git a/Engine/Plugins/Tests/InterchangeTests/Source/InterchangeTests/Public/ImportTestFunctions/SkeletalMeshImportTestFunctions.h b/Engine/Plugins/Tests/InterchangeTests/Source/InterchangeTests/Public/ImportTestFunctions/SkeletalMeshImportTestFunctions.h
--- a/Engine/Plugins/Tests/InterchangeTests/Source/InterchangeTests/Public/ImportTestFunctions/SkeletalMeshImportTestFunctions.h
+++ b/Engine/Plugins/Tests/InterchangeTests/Source/InterchangeTests/Public/ImportTestFunctions/SkeletalMeshImportTestFunctions.h
@@ -5,101 +5,105 @@
 #include "ImportTestFunctionsBase.h"
 #include "SkeletalMeshImportTestFunctions.generated.h"
 
+#define UE_API INTERCHANGETESTS_API
+
 struct FInterchangeTestFunctionResult;
 class USkeletalMesh;
 
 
-UCLASS()
-class INTERCHANGETESTS_API USkeletalMeshImportTestFunctions : public UImportTestFunctionsBase
+UCLASS(MinimalAPI)
+class USkeletalMeshImportTestFunctions : public UImportTestFunctionsBase
 {
 	GENERATED_BODY()
 
 public:
 
 	// UImportTestFunctionsBase interface
-	virtual UClass* GetAssociatedAssetType() const override;
+	UE_API virtual UClass* GetAssociatedAssetType() const override;
 
 	/** Check whether the expected number of skeletal meshes are imported */
 	UFUNCTION(Exec)
-	static FInterchangeTestFunctionResult CheckImportedSkeletalMeshCount(const TArray<USkeletalMesh*>& Meshes, int32 ExpectedNumberOfImportedSkeletalMeshes);
+	static UE_API FInterchangeTestFunctionResult CheckImportedSkeletalMeshCount(const TArray<USkeletalMesh*>& Meshes, int32 ExpectedNumberOfImportedSkeletalMeshes);
 
 	/** Check whether the vertex count in the built render data for the given LOD is as expected */
 	UFUNCTION(Exec)
-	static FInterchangeTestFunctionResult CheckRenderVertexCount(USkeletalMesh* Mesh, int32 LodIndex, int32 ExpectedNumberOfRenderVertices);
+	static UE_API FInterchangeTestFunctionResult CheckRenderVertexCount(USkeletalMesh* Mesh, int32 LodIndex, int32 ExpectedNumberOfRenderVertices);
 
 	/** Check whether the triangle count in the built render data for the given LOD is as expected */
 	UFUNCTION(Exec)
-	static FInterchangeTestFunctionResult CheckRenderTriangleCount(USkeletalMesh* Mesh, int32 LodIndex, int32 ExpectedNumberOfRenderTriangles);
+	static UE_API FInterchangeTestFunctionResult CheckRenderTriangleCount(USkeletalMesh* Mesh, int32 LodIndex, int32 ExpectedNumberOfRenderTriangles);
 
 	/** Check whether the mesh has the expected number of LODs */
 	UFUNCTION(Exec)
-	static FInterchangeTestFunctionResult CheckLodCount(USkeletalMesh* Mesh, int32 ExpectedNumberOfLods);
+	static UE_API FInterchangeTestFunctionResult CheckLodCount(USkeletalMesh* Mesh, int32 ExpectedNumberOfLods);
 
 	/** Check whether the mesh has the expected number of material slots */
 	UFUNCTION(Exec)
-	static FInterchangeTestFunctionResult CheckMaterialSlotCount(USkeletalMesh* Mesh, int32 ExpectedNumberOfMaterialSlots);
+	static UE_API FInterchangeTestFunctionResult CheckMaterialSlotCount(USkeletalMesh* Mesh, int32 ExpectedNumberOfMaterialSlots);
 
 	/** Check whether the built render data for the given mesh LOD has the expected number of sections */
 	UFUNCTION(Exec)
-	static FInterchangeTestFunctionResult CheckSectionCount(USkeletalMesh* Mesh, int32 LodIndex, int32 ExpectedNumberOfSections);
+	static UE_API FInterchangeTestFunctionResult CheckSectionCount(USkeletalMesh* Mesh, int32 LodIndex, int32 ExpectedNumberOfSections);
 
 	/** Check whether the given section in the built render data for the given LOD has the expected number of triangles */
 	UFUNCTION(Exec)
-	static FInterchangeTestFunctionResult CheckTriangleCountInSection(USkeletalMesh* Mesh, int32 LodIndex, int32 SectionIndex, int32 ExpectedNumberOfTriangles);
+	static UE_API FInterchangeTestFunctionResult CheckTriangleCountInSection(USkeletalMesh* Mesh, int32 LodIndex, int32 SectionIndex, int32 ExpectedNumberOfTriangles);
 
 	/** Check whether the mesh has the expected number of UV channels */
 	UFUNCTION(Exec)
-	static FInterchangeTestFunctionResult CheckUVChannelCount(USkeletalMesh* Mesh, int32 LodIndex, int32 ExpectedNumberOfUVChannels);
+	static UE_API FInterchangeTestFunctionResult CheckUVChannelCount(USkeletalMesh* Mesh, int32 LodIndex, int32 ExpectedNumberOfUVChannels);
 
 	/** Check whether the material name for the given section in the render data for the given LOD is as expected */
 	UFUNCTION(Exec)
-	static FInterchangeTestFunctionResult CheckSectionMaterialName(USkeletalMesh* Mesh, int32 LodIndex, int32 SectionIndex, const FString& ExpectedMaterialName);
+	static UE_API FInterchangeTestFunctionResult CheckSectionMaterialName(USkeletalMesh* Mesh, int32 LodIndex, int32 SectionIndex, const FString& ExpectedMaterialName);
 
 	/** Check whether the imported material slot name for the given section in the render data for the given LOD is as expected */
 	UFUNCTION(Exec)
-	static FInterchangeTestFunctionResult CheckSectionImportedMaterialSlotName(USkeletalMesh* Mesh, int32 LodIndex, int32 SectionIndex, const FString& ExpectedImportedMaterialSlotName);
+	static UE_API FInterchangeTestFunctionResult CheckSectionImportedMaterialSlotName(USkeletalMesh* Mesh, int32 LodIndex, int32 SectionIndex, const FString& ExpectedImportedMaterialSlotName);
 
 	/** Check whether the vertex of the given index is at the expected position */
 	UFUNCTION(Exec)
-	static FInterchangeTestFunctionResult CheckVertexIndexPosition(USkeletalMesh* Mesh, int32 LodIndex, int32 VertexIndex, const FVector& ExpectedVertexPosition);
+	static UE_API FInterchangeTestFunctionResult CheckVertexIndexPosition(USkeletalMesh* Mesh, int32 LodIndex, int32 VertexIndex, const FVector& ExpectedVertexPosition);
 
 	/** Check whether the vertex of the given index is at the expected normal */
 	UFUNCTION(Exec)
-	static FInterchangeTestFunctionResult CheckVertexIndexNormal(USkeletalMesh* Mesh, int32 LodIndex, int32 VertexIndex, const FVector& ExpectedVertexNormal);
+	static UE_API FInterchangeTestFunctionResult CheckVertexIndexNormal(USkeletalMesh* Mesh, int32 LodIndex, int32 VertexIndex, const FVector& ExpectedVertexNormal);
 
 	/** Check whether the vertex of the given index is at the expected normal */
 	UFUNCTION(Exec)
-	static FInterchangeTestFunctionResult CheckVertexIndexColor(USkeletalMesh* Mesh, int32 LodIndex, int32 VertexIndex, const FColor& ExpectedVertexColor);
+	static UE_API FInterchangeTestFunctionResult CheckVertexIndexColor(USkeletalMesh* Mesh, int32 LodIndex, int32 VertexIndex, const FColor& ExpectedVertexColor);
 
 	/** Check whether the mesh has the expected number of morph targets */
 	UFUNCTION(Exec)
-	static FInterchangeTestFunctionResult CheckMorphTargetCount(USkeletalMesh* Mesh, int32 ExpectedNumberOfMorphTargets);
+	static UE_API FInterchangeTestFunctionResult CheckMorphTargetCount(USkeletalMesh* Mesh, int32 ExpectedNumberOfMorphTargets);
 
 	/** Check whether the imported morph target name is as expected */
 	UFUNCTION(Exec)
-	static FInterchangeTestFunctionResult CheckMorphTargetName(USkeletalMesh* Mesh, int32 MorphTargetIndex, const FString& ExpectedMorphTargetName);
+	static UE_API FInterchangeTestFunctionResult CheckMorphTargetName(USkeletalMesh* Mesh, int32 MorphTargetIndex, const FString& ExpectedMorphTargetName);
 
 	/** Check whether the mesh has the expected number of bones */
 	UFUNCTION(Exec)
-	static FInterchangeTestFunctionResult CheckBoneCount(USkeletalMesh* Mesh, int32 ExpectedNumberOfBones);
+	static UE_API FInterchangeTestFunctionResult CheckBoneCount(USkeletalMesh* Mesh, int32 ExpectedNumberOfBones);
 
 	/** Check that the bone of the specified index has the expected position */
 	UFUNCTION(Exec)
-	static FInterchangeTestFunctionResult CheckBonePosition(USkeletalMesh* Mesh, int32 BoneIndex, const FVector& ExpectedBonePosition);
+	static UE_API FInterchangeTestFunctionResult CheckBonePosition(USkeletalMesh* Mesh, int32 BoneIndex, const FVector& ExpectedBonePosition);
 
 	/** Check whether the static mesh expected number of sockets were imported */
 	UFUNCTION(Exec)
-	static FInterchangeTestFunctionResult CheckSocketCount(USkeletalMesh* Mesh, int32 ExpectedSocketCount);
+	static UE_API FInterchangeTestFunctionResult CheckSocketCount(USkeletalMesh* Mesh, int32 ExpectedSocketCount);
 
 	/** Check whether the static mesh given socket index has the expected name */
 	UFUNCTION(Exec)
-	static FInterchangeTestFunctionResult CheckSocketName(USkeletalMesh* Mesh, int32 SocketIndex, const FString& ExpectedSocketName);
+	static UE_API FInterchangeTestFunctionResult CheckSocketName(USkeletalMesh* Mesh, int32 SocketIndex, const FString& ExpectedSocketName);
 
 	/** Check whether the static mesh given socket index has the expected location */
 	UFUNCTION(Exec)
-	static FInterchangeTestFunctionResult CheckSocketLocation(USkeletalMesh* Mesh, int32 SocketIndex, const FVector& ExpectedSocketLocation);
+	static UE_API FInterchangeTestFunctionResult CheckSocketLocation(USkeletalMesh* Mesh, int32 SocketIndex, const FVector& ExpectedSocketLocation);
 
 	/** Check that the specified bone name is skinned with the expected number of vertices */
 	UFUNCTION(Exec)
-	static FInterchangeTestFunctionResult CheckSkinnedVertexCountForBone(USkeletalMesh* Mesh, const FString& BoneName, bool bTestFirstAlternateProfile, int32 ExpectedSkinnedVertexCount);
+	static UE_API FInterchangeTestFunctionResult CheckSkinnedVertexCountForBone(USkeletalMesh* Mesh, const FString& BoneName, bool bTestFirstAlternateProfile, int32 ExpectedSkinnedVertexCount);
 };
+
+#undef UE_API
diff --git a/Engine/Plugins/Tests/InterchangeTests/Source/InterchangeTests/Public/ImportTestFunctions/SpotLightImportTestFunctions.h b/Engine/Plugins/Tests/InterchangeTests/Source/InterchangeTests/Public/ImportTestFunctions/SpotLightImportTestFunctions.h
--- a/Engine/Plugins/Tests/InterchangeTests/Source/InterchangeTests/Public/ImportTestFunctions/SpotLightImportTestFunctions.h
+++ b/Engine/Plugins/Tests/InterchangeTests/Source/InterchangeTests/Public/ImportTestFunctions/SpotLightImportTestFunctions.h
@@ -6,24 +6,28 @@
 #include "InterchangeTestFunction.h"
 #include "SpotLightImportTestFunctions.generated.h"
 
+#define UE_API INTERCHANGETESTS_API
+
 struct FInterchangeTestFunctionResult;
 
 
-UCLASS()
-class INTERCHANGETESTS_API USpotLightImportTestFunctions : public UPointLightImportTestFunctions
+UCLASS(MinimalAPI)
+class USpotLightImportTestFunctions : public UPointLightImportTestFunctions
 {
 	GENERATED_BODY()
 
 public:
 
 	// UImportTestFunctionsBase interface
-	virtual UClass* GetAssociatedAssetType() const override;
+	UE_API virtual UClass* GetAssociatedAssetType() const override;
 
 	/** Check whether the light inner cone angle is correct*/
 	UFUNCTION(Exec)
-	static FInterchangeTestFunctionResult CheckLightInnerConeAngle(ASpotLight* Light, float ExpectedLightInnerConeAngle);
+	static UE_API FInterchangeTestFunctionResult CheckLightInnerConeAngle(ASpotLight* Light, float ExpectedLightInnerConeAngle);
 
 	/** Check whether the light outer cone angle is correct*/
 	UFUNCTION(Exec)
-	static FInterchangeTestFunctionResult CheckLightOuterConeAngle(ASpotLight* Light, float ExpectedLightOuterConeAngle);
+	static UE_API FInterchangeTestFunctionResult CheckLightOuterConeAngle(ASpotLight* Light, float ExpectedLightOuterConeAngle);
 };
+
+#undef UE_API
diff --git a/Engine/Plugins/Tests/InterchangeTests/Source/InterchangeTests/Public/ImportTestFunctions/StaticMeshImportTestFunctions.h b/Engine/Plugins/Tests/InterchangeTests/Source/InterchangeTests/Public/ImportTestFunctions/StaticMeshImportTestFunctions.h
--- a/Engine/Plugins/Tests/InterchangeTests/Source/InterchangeTests/Public/ImportTestFunctions/StaticMeshImportTestFunctions.h
+++ b/Engine/Plugins/Tests/InterchangeTests/Source/InterchangeTests/Public/ImportTestFunctions/StaticMeshImportTestFunctions.h
@@ -5,6 +5,8 @@
 #include "ImportTestFunctionsBase.h"
 #include "StaticMeshImportTestFunctions.generated.h"
 
+#define UE_API INTERCHANGETESTS_API
+
 struct FMeshBuildSettings;
 struct FMeshNaniteSettings;
 
@@ -26,131 +28,131 @@ enum class EStaticMeshImportTestGroundTruthBitflags : uint32
 ENUM_CLASS_FLAGS(EStaticMeshImportTestGroundTruthBitflags);
 
 
-UCLASS()
-class INTERCHANGETESTS_API UStaticMeshImportTestFunctions : public UImportTestFunctionsBase
+UCLASS(MinimalAPI)
+class UStaticMeshImportTestFunctions : public UImportTestFunctionsBase
 {
 	GENERATED_BODY()
 
 public:
 
 	// UImportTestFunctionsBase interface
-	virtual UClass* GetAssociatedAssetType() const override;
+	UE_API virtual UClass* GetAssociatedAssetType() const override;
 
 	/** Check whether the expected number of static meshes are imported */
 	UFUNCTION(Exec)
-	static FInterchangeTestFunctionResult CheckImportedStaticMeshCount(const TArray<UStaticMesh*>& Meshes, int32 ExpectedNumberOfImportedStaticMeshes);
+	static UE_API FInterchangeTestFunctionResult CheckImportedStaticMeshCount(const TArray<UStaticMesh*>& Meshes, int32 ExpectedNumberOfImportedStaticMeshes);
 
 	/** Check whether the static mesh has the expected number of LODs */
 	UFUNCTION(Exec)
-	static FInterchangeTestFunctionResult CheckLodCount(UStaticMesh* Mesh, int32 ExpectedNumberOfLods);
+	static UE_API FInterchangeTestFunctionResult CheckLodCount(UStaticMesh* Mesh, int32 ExpectedNumberOfLods);
 
 	/** Check whether the render data for the given mesh LOD has the expected number of sections */
 	UFUNCTION(Exec)
-	static FInterchangeTestFunctionResult CheckSectionCount(UStaticMesh* Mesh, int32 LodIndex, int32 ExpectedNumberOfSections);
+	static UE_API FInterchangeTestFunctionResult CheckSectionCount(UStaticMesh* Mesh, int32 LodIndex, int32 ExpectedNumberOfSections);
 
 	/** Check whether the static mesh has the expected number of material slots */
 	UFUNCTION(Exec)
-	static FInterchangeTestFunctionResult CheckMaterialSlotCount(UStaticMesh* Mesh, int32 ExpectedNumberOfMaterialSlots);
+	static UE_API FInterchangeTestFunctionResult CheckMaterialSlotCount(UStaticMesh* Mesh, int32 ExpectedNumberOfMaterialSlots);
 
 	/** Check whether the static mesh imported material slot name for the given material index is as expected */
 	UFUNCTION(Exec)
-	static FInterchangeTestFunctionResult CheckImportedMaterialSlotName(UStaticMesh* Mesh, int32 MaterialIndex, const FString& ExpectedImportedMaterialSlotName);
+	static UE_API FInterchangeTestFunctionResult CheckImportedMaterialSlotName(UStaticMesh* Mesh, int32 MaterialIndex, const FString& ExpectedImportedMaterialSlotName);
 
 	/** Check whether the mesh description has the expected number of polygon groups for the given LOD */
 	UFUNCTION(Exec)
-	static FInterchangeTestFunctionResult CheckPolygonGroupCount(UStaticMesh* Mesh, int32 LodIndex, int32 ExpectedNumberOfPolygonGroups);
+	static UE_API FInterchangeTestFunctionResult CheckPolygonGroupCount(UStaticMesh* Mesh, int32 LodIndex, int32 ExpectedNumberOfPolygonGroups);
 
 	/** Check whether the mesh description material slot name for the given polygon group in the given LOD is as expected */
 	UFUNCTION(Exec)
-	static FInterchangeTestFunctionResult CheckPolygonGroupImportedMaterialSlotName(UStaticMesh* Mesh, int32 LodIndex, int32 PolygonGroupIndex, const FString& ExpectedImportedMaterialSlotName);
+	static UE_API FInterchangeTestFunctionResult CheckPolygonGroupImportedMaterialSlotName(UStaticMesh* Mesh, int32 LodIndex, int32 PolygonGroupIndex, const FString& ExpectedImportedMaterialSlotName);
 
 	/** Check whether the section index in the built render data for the given LOD is referencing the expected material index */
 	UFUNCTION(Exec)
-	static FInterchangeTestFunctionResult CheckSectionMaterialIndex(UStaticMesh* Mesh, int32 LodIndex, int32 SectionIndex, int32 ExpectedMaterialIndex);
+	static UE_API FInterchangeTestFunctionResult CheckSectionMaterialIndex(UStaticMesh* Mesh, int32 LodIndex, int32 SectionIndex, int32 ExpectedMaterialIndex);
 
 	/** Check whether the section index in the built render data for the given LOD is referencing the expected material */
 	UFUNCTION(Exec)
-	static FInterchangeTestFunctionResult CheckSectionMaterialName(UStaticMesh* Mesh, int32 LodIndex, int32 SectionIndex, const FString& ExpectedMaterialName);
+	static UE_API FInterchangeTestFunctionResult CheckSectionMaterialName(UStaticMesh* Mesh, int32 LodIndex, int32 SectionIndex, const FString& ExpectedMaterialName);
 
 	/** Check whether the section index in the built render data for the given LOD is referencing the expected imported material slot name*/
 	UFUNCTION(Exec)
-	static FInterchangeTestFunctionResult CheckSectionImportedMaterialSlotName(UStaticMesh* Mesh, int32 LodIndex, int32 SectionIndex, const FString& ExpectedImportedMaterialSlotName);
+	static UE_API FInterchangeTestFunctionResult CheckSectionImportedMaterialSlotName(UStaticMesh* Mesh, int32 LodIndex, int32 SectionIndex, const FString& ExpectedImportedMaterialSlotName);
 
 	/** Check whether the mesh description vertex count for the given LOD is as expected */
 	UFUNCTION(Exec)
-	static FInterchangeTestFunctionResult CheckVertexCount(UStaticMesh* Mesh, int32 LodIndex, int32 ExpectedNumberOfVertices);
+	static UE_API FInterchangeTestFunctionResult CheckVertexCount(UStaticMesh* Mesh, int32 LodIndex, int32 ExpectedNumberOfVertices);
 
 	/** Check whether the render data vertex count for the given LOD is as expected */
 	UFUNCTION(Exec)
-	static FInterchangeTestFunctionResult CheckRenderVertexCount(UStaticMesh* Mesh, int32 LodIndex, int32 ExpectedNumberOfRenderVertices);
+	static UE_API FInterchangeTestFunctionResult CheckRenderVertexCount(UStaticMesh* Mesh, int32 LodIndex, int32 ExpectedNumberOfRenderVertices);
 
 	/** Check whether the mesh description vertex of the given index is at the expected position */
 	UFUNCTION(Exec)
-	static FInterchangeTestFunctionResult CheckVertexIndexPosition(UStaticMesh* Mesh, int32 LodIndex, int32 VertexIndex, const FVector& ExpectedVertexPosition);
+	static UE_API FInterchangeTestFunctionResult CheckVertexIndexPosition(UStaticMesh* Mesh, int32 LodIndex, int32 VertexIndex, const FVector& ExpectedVertexPosition);
 
 	/** Check whether the render data vertex index of the given LOD has the expected vertex normal */
 	UFUNCTION(Exec)
-	static FInterchangeTestFunctionResult CheckRenderVertexIndexNormal(UStaticMesh* Mesh, int32 LodIndex, int32 VertexIndex, const FVector4f& ExpectedVertexNormal);
+	static UE_API FInterchangeTestFunctionResult CheckRenderVertexIndexNormal(UStaticMesh* Mesh, int32 LodIndex, int32 VertexIndex, const FVector4f& ExpectedVertexNormal);
 
 	/** Check whether the render data vertex index of the given LOD has the expected vertex color */
 	UFUNCTION(Exec)
-	static FInterchangeTestFunctionResult CheckRenderVertexIndexColor(UStaticMesh* Mesh, int32 LodIndex, int32 VertexIndex, const FColor& ExpectedVertexColor);
+	static UE_API FInterchangeTestFunctionResult CheckRenderVertexIndexColor(UStaticMesh* Mesh, int32 LodIndex, int32 VertexIndex, const FColor& ExpectedVertexColor);
 
 	/** Check whether the render data contains any vertex colors of the given LOD */
 	UFUNCTION(Exec)
-	static FInterchangeTestFunctionResult CheckRenderHasVertexColors(UStaticMesh* Mesh, int32 LodIndex, bool bExpectedHasVertexColors);
+	static UE_API FInterchangeTestFunctionResult CheckRenderHasVertexColors(UStaticMesh* Mesh, int32 LodIndex, bool bExpectedHasVertexColors);
 
 	/** Check whether the mesh description has the expected number of triangles for the given LOD */
 	UFUNCTION(Exec)
-	static FInterchangeTestFunctionResult CheckTriangleCount(UStaticMesh* Mesh, int32 LodIndex, int32 ExpectedTotalNumberOfTriangles);
+	static UE_API FInterchangeTestFunctionResult CheckTriangleCount(UStaticMesh* Mesh, int32 LodIndex, int32 ExpectedTotalNumberOfTriangles);
 
 	/** Check whether the render data has the expected number of triangles for the given LOD */
 	UFUNCTION(Exec)
-	static FInterchangeTestFunctionResult CheckRenderTriangleCount(UStaticMesh* Mesh, int32 LodIndex, int32 ExpectedTotalNumberOfTriangles);
+	static UE_API FInterchangeTestFunctionResult CheckRenderTriangleCount(UStaticMesh* Mesh, int32 LodIndex, int32 ExpectedTotalNumberOfTriangles);
 
 	/** Check whether the mesh description has the expected number of polygons for the given LOD */
 	UFUNCTION(Exec)
-	static FInterchangeTestFunctionResult CheckPolygonCount(UStaticMesh* Mesh, int32 LodIndex, int32 ExpectedNumberOfPolygons);
+	static UE_API FInterchangeTestFunctionResult CheckPolygonCount(UStaticMesh* Mesh, int32 LodIndex, int32 ExpectedNumberOfPolygons);
 
 	/** Check whether the mesh description contains any quads or ngons */
 	UFUNCTION(Exec)
-	static FInterchangeTestFunctionResult CheckThatMeshHasQuadsOrNgons(UStaticMesh* Mesh, int32 LodIndex, bool bMeshHasQuadsOrNgons);
+	static UE_API FInterchangeTestFunctionResult CheckThatMeshHasQuadsOrNgons(UStaticMesh* Mesh, int32 LodIndex, bool bMeshHasQuadsOrNgons);
 
 	/** Check whether the given mesh description polygon group of the given LOD has the expected number of triangles */
 	UFUNCTION(Exec)
-	static FInterchangeTestFunctionResult CheckTriangleCountInPolygonGroup(UStaticMesh* Mesh, int32 LodIndex, int32 PolygonGroupIndex, int32 ExpectedNumberOfTriangles);
+	static UE_API FInterchangeTestFunctionResult CheckTriangleCountInPolygonGroup(UStaticMesh* Mesh, int32 LodIndex, int32 PolygonGroupIndex, int32 ExpectedNumberOfTriangles);
 
 	/** Check whether the given mesh description polygon group of the given LOD has the expected number of polygons */
 	UFUNCTION(Exec)
-	static FInterchangeTestFunctionResult CheckPolygonCountInPolygonGroup(UStaticMesh* Mesh, int32 LodIndex, int32 PolygonGroupIndex, int32 ExpectedNumberOfPolygons);
+	static UE_API FInterchangeTestFunctionResult CheckPolygonCountInPolygonGroup(UStaticMesh* Mesh, int32 LodIndex, int32 PolygonGroupIndex, int32 ExpectedNumberOfPolygons);
 
 	/** Check whether the mesh description has the expected number of UV channels */
 	UFUNCTION(Exec)
-	static FInterchangeTestFunctionResult CheckUVChannelCount(UStaticMesh* Mesh, int32 LodIndex, int32 ExpectedNumberOfUVChannels);
+	static UE_API FInterchangeTestFunctionResult CheckUVChannelCount(UStaticMesh* Mesh, int32 LodIndex, int32 ExpectedNumberOfUVChannels);
 
 	/** Check whether the render data has the expected number of UV channels */
 	UFUNCTION(Exec)
-	static FInterchangeTestFunctionResult CheckRenderUVChannelCount(UStaticMesh* Mesh, int32 LodIndex, int32 ExpectedNumberOfUVChannels);
+	static UE_API FInterchangeTestFunctionResult CheckRenderUVChannelCount(UStaticMesh* Mesh, int32 LodIndex, int32 ExpectedNumberOfUVChannels);
 
 	/** Check whether the static mesh expected number of simple collision primitives were imported */
 	UFUNCTION(Exec)
-	static FInterchangeTestFunctionResult CheckSimpleCollisionPrimitiveCount(UStaticMesh* Mesh, int32 ExpectedSphereElementCount, int32 ExpectedBoxElementCount, int32 ExpectedCapsuleElementCount, int32 ExpectedConvexElementCount, int32 ExpectedTaperedCapsuleElementCount);
+	static UE_API FInterchangeTestFunctionResult CheckSimpleCollisionPrimitiveCount(UStaticMesh* Mesh, int32 ExpectedSphereElementCount, int32 ExpectedBoxElementCount, int32 ExpectedCapsuleElementCount, int32 ExpectedConvexElementCount, int32 ExpectedTaperedCapsuleElementCount);
 
 	/** Check whether the static mesh expected number of sockets were imported */
 	UFUNCTION(Exec)
-	static FInterchangeTestFunctionResult CheckSocketCount(UStaticMesh* Mesh, int32 ExpectedSocketCount);
+	static UE_API FInterchangeTestFunctionResult CheckSocketCount(UStaticMesh* Mesh, int32 ExpectedSocketCount);
 
 	/** Check whether the static mesh given socket index has the expected name */
 	UFUNCTION(Exec)
-	static FInterchangeTestFunctionResult CheckSocketName(UStaticMesh* Mesh, int32 SocketIndex, const FString& ExpectedSocketName);
+	static UE_API FInterchangeTestFunctionResult CheckSocketName(UStaticMesh* Mesh, int32 SocketIndex, const FString& ExpectedSocketName);
 
 	/** Check whether the static mesh given socket index has the expected location */
 	UFUNCTION(Exec)
-	static FInterchangeTestFunctionResult CheckSocketLocation(UStaticMesh* Mesh, int32 SocketIndex, const FVector& ExpectedSocketLocation);
+	static UE_API FInterchangeTestFunctionResult CheckSocketLocation(UStaticMesh* Mesh, int32 SocketIndex, const FVector& ExpectedSocketLocation);
 
 	/** Check whether the static mesh is equivalent to a ground truth asset */
 	UFUNCTION(Exec)
-	static FInterchangeTestFunctionResult CheckAgainstGroundTruth(UStaticMesh* Mesh, TSoftObjectPtr<UStaticMesh> MeshToCompare,
+	static UE_API FInterchangeTestFunctionResult CheckAgainstGroundTruth(UStaticMesh* Mesh, TSoftObjectPtr<UStaticMesh> MeshToCompare,
 		bool bCheckVertexCountEqual = true,
 		bool bCheckTriangleCountEqual = true,
 		bool bCheckUVChannelCountEqual = true,
@@ -161,21 +163,23 @@ public:
 
 	/** Check whether the static mesh source model has the expected build settings */
 	UFUNCTION(Exec)
-	static FInterchangeTestFunctionResult CheckBuildSettings(UStaticMesh* Mesh, int32 LodIndex, const FMeshBuildSettings& ExpectedBuildSettings);
+	static UE_API FInterchangeTestFunctionResult CheckBuildSettings(UStaticMesh* Mesh, int32 LodIndex, const FMeshBuildSettings& ExpectedBuildSettings);
 
 	/** Check whether the static mesh has the expected Nanite settings */
 	UFUNCTION(Exec)
-	static FInterchangeTestFunctionResult CheckNaniteSettings(UStaticMesh* Mesh, const FMeshNaniteSettings& ExpectedNaniteSettings);
+	static UE_API FInterchangeTestFunctionResult CheckNaniteSettings(UStaticMesh* Mesh, const FMeshNaniteSettings& ExpectedNaniteSettings);
 
 	/** Check whether the bounding box size is greater than given value */
 	UFUNCTION(Exec)
-	static FInterchangeTestFunctionResult CheckBoundingBoxSizeGreaterThan(UStaticMesh* Mesh, const FVector& GreaterThanSize);
+	static UE_API FInterchangeTestFunctionResult CheckBoundingBoxSizeGreaterThan(UStaticMesh* Mesh, const FVector& GreaterThanSize);
 
 	/** Check whether the bounding box size is less than given value */
 	UFUNCTION(Exec)
-	static FInterchangeTestFunctionResult CheckBoundingBoxSizeLessThan(UStaticMesh* Mesh, const FVector& LessThanSize);
+	static UE_API FInterchangeTestFunctionResult CheckBoundingBoxSizeLessThan(UStaticMesh* Mesh, const FVector& LessThanSize);
 
 	/** Check whether the bounding box size is equal to given value */
 	UFUNCTION(Exec)
-	static FInterchangeTestFunctionResult CheckBoundingBoxSize(UStaticMesh* Mesh, const FVector& ExpectedSize);
+	static UE_API FInterchangeTestFunctionResult CheckBoundingBoxSize(UStaticMesh* Mesh, const FVector& ExpectedSize);
 };
+
+#undef UE_API
diff --git a/Engine/Plugins/Tests/InterchangeTests/Source/InterchangeTests/Public/ImportTestFunctions/TextureImportTestFunctions.h b/Engine/Plugins/Tests/InterchangeTests/Source/InterchangeTests/Public/ImportTestFunctions/TextureImportTestFunctions.h
--- a/Engine/Plugins/Tests/InterchangeTests/Source/InterchangeTests/Public/ImportTestFunctions/TextureImportTestFunctions.h
+++ b/Engine/Plugins/Tests/InterchangeTests/Source/InterchangeTests/Public/ImportTestFunctions/TextureImportTestFunctions.h
@@ -6,38 +6,42 @@
 #include "ImportTestFunctionsBase.h"
 #include "TextureImportTestFunctions.generated.h"
 
+#define UE_API INTERCHANGETESTS_API
+
 class UTexture;
 
 struct FInterchangeTestFunctionResult;
 
 
-UCLASS()
-class INTERCHANGETESTS_API UTextureImportTestFunctions : public UImportTestFunctionsBase
+UCLASS(MinimalAPI)
+class UTextureImportTestFunctions : public UImportTestFunctionsBase
 {
 	GENERATED_BODY()
 
 public:
 
 	// UImportTestFunctionsBase interface
-	virtual UClass* GetAssociatedAssetType() const override;
+	UE_API virtual UClass* GetAssociatedAssetType() const override;
 
 	/** Check whether the expected number of textures are imported */
 	UFUNCTION(Exec)
-	static FInterchangeTestFunctionResult CheckImportedTextureCount(const TArray<UTexture*>& Textures, int32 ExpectedNumberOfImportedTextures);
+	static UE_API FInterchangeTestFunctionResult CheckImportedTextureCount(const TArray<UTexture*>& Textures, int32 ExpectedNumberOfImportedTextures);
 
 	/** Check whether the imported texture has the expected filtering mode */
 	UFUNCTION(Exec)
-	static FInterchangeTestFunctionResult CheckTextureFilter(const UTexture* Texture, TextureFilter ExpectedTextureFilter);
+	static UE_API FInterchangeTestFunctionResult CheckTextureFilter(const UTexture* Texture, TextureFilter ExpectedTextureFilter);
 
 	/** Check whether the imported texture has the expected addressing mode for X-axis */
 	UFUNCTION(Exec)
-	static FInterchangeTestFunctionResult CheckTextureAddressX(const UTexture* Texture, TextureAddress ExpectedTextureAddressX);
+	static UE_API FInterchangeTestFunctionResult CheckTextureAddressX(const UTexture* Texture, TextureAddress ExpectedTextureAddressX);
 
 	/** Check whether the imported texture has the expected addressing mode for Y-axis */
 	UFUNCTION(Exec)
-	static FInterchangeTestFunctionResult CheckTextureAddressY(const UTexture* Texture, TextureAddress ExpectedTextureAddressY);
+	static UE_API FInterchangeTestFunctionResult CheckTextureAddressY(const UTexture* Texture, TextureAddress ExpectedTextureAddressY);
 
 	/** Check whether the imported texture has the expected addressing mode for Z-axis */
 	UFUNCTION(Exec)
-	static FInterchangeTestFunctionResult CheckTextureAddressZ(const UTexture* Texture, TextureAddress ExpectedTextureAddressZ);
+	static UE_API FInterchangeTestFunctionResult CheckTextureAddressZ(const UTexture* Texture, TextureAddress ExpectedTextureAddressZ);
 };
+
+#undef UE_API
diff --git a/Engine/Plugins/Tests/InterchangeTests/Source/InterchangeTests/Public/InterchangeImportTestPlan.h b/Engine/Plugins/Tests/InterchangeTests/Source/InterchangeTests/Public/InterchangeImportTestPlan.h
--- a/Engine/Plugins/Tests/InterchangeTests/Source/InterchangeTests/Public/InterchangeImportTestPlan.h
+++ b/Engine/Plugins/Tests/InterchangeTests/Source/InterchangeTests/Public/InterchangeImportTestPlan.h
@@ -8,6 +8,8 @@
 
 #include "InterchangeImportTestPlan.generated.h"
 
+#define UE_API INTERCHANGETESTS_API
+
 class UInterchangeImportTestStepBase;
 class UInterchangeImportTestStepImport;
 class UInterchangeImportTestStepReimport;
@@ -15,21 +17,21 @@ class UInterchangeImportTestStepReimport;
 /**
 * Define a test plan
 */
-UCLASS(BlueprintType, EditInlineNew, meta=(PrioritizeCategories = "Description Level Definition Import Reimport")) // @TODO: can't find a way to force Run to be below Definition
-class INTERCHANGETESTS_API UInterchangeImportTestPlan : public UObject
+UCLASS(MinimalAPI, BlueprintType, EditInlineNew, meta=(PrioritizeCategories = "Description Level Definition Import Reimport")) // @TODO: can't find a way to force Run to be below Definition
+class UInterchangeImportTestPlan : public UObject
 {
 	GENERATED_BODY()
 
 public:
 
 	/** Deserialize a test plan from Json */
-	bool ReadFromJson(const FString& Filename);
+	UE_API bool ReadFromJson(const FString& Filename);
 
 	/* Serialize a test plan to Json */
-	void WriteToJson(const FString& Filename);
+	UE_API void WriteToJson(const FString& Filename);
 
 public:
-	UInterchangeImportTestPlan();
+	UE_API UInterchangeImportTestPlan();
 
 	/** Test Description */
 	UPROPERTY(EditAnywhere, Category = "Description", Meta = (MultiLine = true))
@@ -55,10 +57,10 @@ public:
 public:
 	/** Click here to immediately run this single test through the automation framework */
 	UFUNCTION(BlueprintCallable, CallInEditor, Category = "Automation")
-	void RunThisTest();	
-	bool IsRunningSynchornously() const;
+	UE_API void RunThisTest();	
+	UE_API bool IsRunningSynchornously() const;
 
-	void SetupLevelForImport();
+	UE_API void SetupLevelForImport();
 	
 	UWorld* GetCurrentWorld() const
 	{	
@@ -70,19 +72,19 @@ public:
 		return GetCurrentWorld()->GetCurrentLevel();
 	}
 
-	void CleanupLevel();
+	UE_API void CleanupLevel();
 
-	virtual void PostLoad() override;
-	void OnAssetLoaded(UObject* Asset);
+	UE_API virtual void PostLoad() override;
+	UE_API void OnAssetLoaded(UObject* Asset);
 
 #if WITH_EDITOR
-	static FSuppressableWarningDialog::EMode GetInterchangeTestPlanWarningDialogMode();
+	static UE_API FSuppressableWarningDialog::EMode GetInterchangeTestPlanWarningDialogMode();
 
-	virtual void PreEditUndo() override;
-	virtual void PostEditUndo() override;
+	UE_API virtual void PreEditUndo() override;
+	UE_API virtual void PostEditUndo() override;
 
-	virtual void PreEditChange(FProperty* PropertyAboutToChange) override;
-	virtual void PostEditChangeChainProperty(FPropertyChangedChainEvent& PropertyChangedEvent) override;
+	UE_API virtual void PreEditChange(FProperty* PropertyAboutToChange) override;
+	UE_API virtual void PostEditChangeChainProperty(FPropertyChangedChainEvent& PropertyChangedEvent) override;
 #endif
 	
 private:
@@ -96,11 +98,13 @@ private:
 
 namespace UE::Interchange
 {
-	class INTERCHANGETESTS_API FInterchangeImportTestPlanStaticHelpers
+	class FInterchangeImportTestPlanStaticHelpers
 	{
 	public:
-		static FString GetTestNameFromObjectPathString(const FString& InObjectPathString, bool bAddBeautifiedTestNamePrefix = false);
-		static FString GetBeautifiedTestName();
-		static FString GetInterchangeImportTestRootGameFolder();
+		static UE_API FString GetTestNameFromObjectPathString(const FString& InObjectPathString, bool bAddBeautifiedTestNamePrefix = false);
+		static UE_API FString GetBeautifiedTestName();
+		static UE_API FString GetInterchangeImportTestRootGameFolder();
 	};
-}
\ No newline at end of file
+}
+
+#undef UE_API
diff --git a/Engine/Plugins/Tests/InterchangeTests/Source/InterchangeTests/Public/InterchangeImportTestSettings.h b/Engine/Plugins/Tests/InterchangeTests/Source/InterchangeTests/Public/InterchangeImportTestSettings.h
--- a/Engine/Plugins/Tests/InterchangeTests/Source/InterchangeTests/Public/InterchangeImportTestSettings.h
+++ b/Engine/Plugins/Tests/InterchangeTests/Source/InterchangeTests/Public/InterchangeImportTestSettings.h
@@ -8,8 +8,8 @@
 /**
  * Implement settings for the Interchange Import Test
  */
-UCLASS(config=Engine, defaultconfig)
-class INTERCHANGETESTS_API UInterchangeImportTestSettings : public UObject
+UCLASS(MinimalAPI, config=Engine, defaultconfig)
+class UInterchangeImportTestSettings : public UObject
 {
 	GENERATED_BODY()
 
diff --git a/Engine/Plugins/Tests/InterchangeTests/Source/InterchangeTests/Public/InterchangeImportTestStepBase.h b/Engine/Plugins/Tests/InterchangeTests/Source/InterchangeTests/Public/InterchangeImportTestStepBase.h
--- a/Engine/Plugins/Tests/InterchangeTests/Source/InterchangeTests/Public/InterchangeImportTestStepBase.h
+++ b/Engine/Plugins/Tests/InterchangeTests/Source/InterchangeTests/Public/InterchangeImportTestStepBase.h
@@ -9,6 +9,8 @@
 #include "UObject/ObjectPtr.h"
 #include "InterchangeImportTestStepBase.generated.h"
 
+#define UE_API INTERCHANGETESTS_API
+
 class UInterchangeImportTestPlan;
 class FAutomationTestExecutionInfo;
 struct FInterchangeImportTestData;
@@ -53,8 +55,8 @@ public:
 };
 
 
-UCLASS(BlueprintType, EditInlineNew, Abstract, autoExpandCategories = (General, Test))
-class INTERCHANGETESTS_API UInterchangeImportTestStepBase : public UObject
+UCLASS(MinimalAPI, BlueprintType, EditInlineNew, Abstract, autoExpandCategories = (General, Test))
+class UInterchangeImportTestStepBase : public UObject
 {
 	GENERATED_BODY()
 
@@ -77,7 +79,9 @@ public:
 	virtual void ClearPipelineSettings() PURE_VIRTUAL(UInterchangeImportTestStepBase::ClearPipelineSettings);
 	virtual bool IsUsingOverridePipelines(bool bCheckForValidPipelines) const PURE_VIRTUAL(UInterchangeImportTestStepBase::IsUsingOverridePipelines, return false; );
 
-	bool PerformTests(FInterchangeImportTestData& Data, FAutomationTestBase* CurrentTest);
+	UE_API bool PerformTests(FInterchangeImportTestData& Data, FAutomationTestBase* CurrentTest);
 protected:
-	void SaveReloadAssets(FInterchangeImportTestData& Data);
+	UE_API void SaveReloadAssets(FInterchangeImportTestData& Data);
 };
+
+#undef UE_API
diff --git a/Engine/Plugins/Tests/InterchangeTests/Source/InterchangeTests/Public/InterchangeImportTestStepImport.h b/Engine/Plugins/Tests/InterchangeTests/Source/InterchangeTests/Public/InterchangeImportTestStepImport.h
--- a/Engine/Plugins/Tests/InterchangeTests/Source/InterchangeTests/Public/InterchangeImportTestStepImport.h
+++ b/Engine/Plugins/Tests/InterchangeTests/Source/InterchangeTests/Public/InterchangeImportTestStepImport.h
@@ -10,6 +10,8 @@
 
 #include "InterchangeImportTestStepImport.generated.h"
 
+#define UE_API INTERCHANGETESTS_API
+
 class UInterchangeTranslatorBase;
 class UInterchangePipelineBase;
 class UInterchangeBaseNodeContainer;
@@ -36,13 +38,13 @@ public:
 DECLARE_MULTICAST_DELEGATE_OneParam(FOnImportTestStepDataChanged, FImportStepChangedData)
 
 
-UCLASS(BlueprintType, Meta = (DisplayName = "Import a file"))
-class INTERCHANGETESTS_API UInterchangeImportTestStepImport : public UInterchangeImportTestStepBase
+UCLASS(MinimalAPI, BlueprintType, Meta = (DisplayName = "Import a file"))
+class UInterchangeImportTestStepImport : public UInterchangeImportTestStepBase
 {
 	GENERATED_BODY()
 
 public:
-	UInterchangeImportTestStepImport();
+	UE_API UInterchangeImportTestStepImport();
 
 	/** The source file to import (path relative to the json script) */
 	UPROPERTY(EditAnywhere, Category = General)
@@ -78,33 +80,35 @@ public:
 	FOnImportTestStepDataChanged OnImportTestStepDataChanged;
 public:
 	// UInterchangeImportTestStepBase interface
-	virtual TTuple<UE::Interchange::FAssetImportResultPtr, UE::Interchange::FSceneImportResultPtr>
+	UE_API virtual TTuple<UE::Interchange::FAssetImportResultPtr, UE::Interchange::FSceneImportResultPtr>
 		StartStep(FInterchangeImportTestData& Data) override;
-	virtual FTestStepResults FinishStep(FInterchangeImportTestData& Data, FAutomationTestBase* CurrentTest) override;
-	virtual FString GetContextString() const override;
-	virtual bool HasScreenshotTest() const override;
-	virtual FInterchangeTestScreenshotParameters GetScreenshotParameters() const override;
+	UE_API virtual FTestStepResults FinishStep(FInterchangeImportTestData& Data, FAutomationTestBase* CurrentTest) override;
+	UE_API virtual FString GetContextString() const override;
+	UE_API virtual bool HasScreenshotTest() const override;
+	UE_API virtual FInterchangeTestScreenshotParameters GetScreenshotParameters() const override;
 
-	virtual bool CanEditPipelineSettings() const override;
-	virtual void EditPipelineSettings() override;
-	virtual void ClearPipelineSettings() override;
-	virtual bool IsUsingOverridePipelines(bool bCheckForValidPipelines) const override;
+	UE_API virtual bool CanEditPipelineSettings() const override;
+	UE_API virtual void EditPipelineSettings() override;
+	UE_API virtual void ClearPipelineSettings() override;
+	UE_API virtual bool IsUsingOverridePipelines(bool bCheckForValidPipelines) const override;
 
-	bool ShouldImportIntoLevelChangeRequireMessageBox() const;
+	UE_API bool ShouldImportIntoLevelChangeRequireMessageBox() const;
 
-	virtual void PostLoad() override;
+	UE_API virtual void PostLoad() override;
 
 #if WITH_EDITOR
-	virtual void PostEditChangeProperty(FPropertyChangedEvent& PropertyChangedEvent) override;
-	virtual void PostEditChangeChainProperty(FPropertyChangedChainEvent& PropertyChangedEvent) override;
+	UE_API virtual void PostEditChangeProperty(FPropertyChangedEvent& PropertyChangedEvent) override;
+	UE_API virtual void PostEditChangeChainProperty(FPropertyChangedChainEvent& PropertyChangedEvent) override;
 #endif
 	
-	TArray<TObjectPtr<UInterchangePipelineBase>> GetCurrentPipelinesOrDefault() const;
+	UE_API TArray<TObjectPtr<UInterchangePipelineBase>> GetCurrentPipelinesOrDefault() const;
 private:
-	void BroadcastImportStepChangedEvent(EImportStepDataChangeType ChangeType);
+	UE_API void BroadcastImportStepChangedEvent(EImportStepDataChangeType ChangeType);
 
 private:
 	UPROPERTY()
 	FString LastSourceFileExtension;
 
 };
+
+#undef UE_API
diff --git a/Engine/Plugins/Tests/InterchangeTests/Source/InterchangeTests/Public/InterchangeImportTestStepReimport.h b/Engine/Plugins/Tests/InterchangeTests/Source/InterchangeTests/Public/InterchangeImportTestStepReimport.h
--- a/Engine/Plugins/Tests/InterchangeTests/Source/InterchangeTests/Public/InterchangeImportTestStepReimport.h
+++ b/Engine/Plugins/Tests/InterchangeTests/Source/InterchangeTests/Public/InterchangeImportTestStepReimport.h
@@ -10,17 +10,19 @@
 #include "Delegates/DelegateCombinations.h"
 #include "InterchangeImportTestStepReimport.generated.h"
 
+#define UE_API INTERCHANGETESTS_API
+
 class UInterchangeImportTestPlan;
 class UInterchangePipelineBase;
 struct FInterchangeTestFunction;
 
-UCLASS(BlueprintType, Meta = (DisplayName = "Reimport a file"), AutoExpandCategories="General")
-class INTERCHANGETESTS_API UInterchangeImportTestStepReimport : public UInterchangeImportTestStepBase
+UCLASS(MinimalAPI, BlueprintType, Meta = (DisplayName = "Reimport a file"), AutoExpandCategories="General")
+class UInterchangeImportTestStepReimport : public UInterchangeImportTestStepBase
 {
 	GENERATED_BODY()
 
 public:
-	UInterchangeImportTestStepReimport();
+	UE_API UInterchangeImportTestStepReimport();
 
 	/** The source file to import (path relative to the json script). */
 	UPROPERTY(EditAnywhere, Category = General)
@@ -59,29 +61,29 @@ public:
 	FInterchangeTestScreenshotParameters ScreenshotParameters;
 
 public:
-	void InitializeReimportStepFromImportStep(UInterchangeImportTestStepImport* ImportTestStep);
-	void RemoveImportStepPipelineSettingsModifiedDelegate();
+	UE_API void InitializeReimportStepFromImportStep(UInterchangeImportTestStepImport* ImportTestStep);
+	UE_API void RemoveImportStepPipelineSettingsModifiedDelegate();
 
 	// UInterchangeImportTestStepBase interface
-	virtual TTuple<UE::Interchange::FAssetImportResultPtr, UE::Interchange::FSceneImportResultPtr>
+	UE_API virtual TTuple<UE::Interchange::FAssetImportResultPtr, UE::Interchange::FSceneImportResultPtr>
 		StartStep(FInterchangeImportTestData& Data) override;
-	virtual FTestStepResults FinishStep(FInterchangeImportTestData& Data, FAutomationTestBase* CurrentTest) override;
-	virtual FString GetContextString() const override;
-	virtual bool HasScreenshotTest() const override;
-	virtual FInterchangeTestScreenshotParameters GetScreenshotParameters() const override;
+	UE_API virtual FTestStepResults FinishStep(FInterchangeImportTestData& Data, FAutomationTestBase* CurrentTest) override;
+	UE_API virtual FString GetContextString() const override;
+	UE_API virtual bool HasScreenshotTest() const override;
+	UE_API virtual FInterchangeTestScreenshotParameters GetScreenshotParameters() const override;
 
-	virtual bool CanEditPipelineSettings() const override;
-	virtual void EditPipelineSettings() override;
-	virtual void ClearPipelineSettings() override;
-	virtual bool IsUsingOverridePipelines(bool bCheckForValidPipelines) const override;
+	UE_API virtual bool CanEditPipelineSettings() const override;
+	UE_API virtual void EditPipelineSettings() override;
+	UE_API virtual void ClearPipelineSettings() override;
+	UE_API virtual bool IsUsingOverridePipelines(bool bCheckForValidPipelines) const override;
 
-	void HandleImportPipelineSettingsModified(FImportStepChangedData ChangedData);
+	UE_API void HandleImportPipelineSettingsModified(FImportStepChangedData ChangedData);
 
-	FString GetReimportStepSourceFilePathString();
+	UE_API FString GetReimportStepSourceFilePathString();
 
 #if WITH_EDITOR
-	virtual void PostEditChangeProperty(FPropertyChangedEvent& PropertyChangedEvent) override;
-	virtual void PostEditChangeChainProperty(FPropertyChangedChainEvent& PropertyChangedEvent) override;
+	UE_API virtual void PostEditChangeProperty(FPropertyChangedEvent& PropertyChangedEvent) override;
+	UE_API virtual void PostEditChangeChainProperty(FPropertyChangedChainEvent& PropertyChangedEvent) override;
 #endif
 
 private:
@@ -90,3 +92,5 @@ private:
 
 	TWeakObjectPtr<UInterchangeImportTestStepImport> CachedImportStep = nullptr;
 };
+
+#undef UE_API
diff --git a/Engine/Plugins/Tests/InterchangeTests/Source/InterchangeTests/Public/InterchangeTestFunction.h b/Engine/Plugins/Tests/InterchangeTests/Source/InterchangeTests/Public/InterchangeTestFunction.h
--- a/Engine/Plugins/Tests/InterchangeTests/Source/InterchangeTests/Public/InterchangeTestFunction.h
+++ b/Engine/Plugins/Tests/InterchangeTests/Source/InterchangeTests/Public/InterchangeTestFunction.h
@@ -5,6 +5,8 @@
 #include "Templates/SubclassOf.h"
 #include "InterchangeTestFunction.generated.h"
 
+#define UE_API INTERCHANGETESTS_API
+
 class FStructOnScope;
 
 
@@ -49,7 +51,7 @@ private:
 
 
 USTRUCT()
-struct INTERCHANGETESTS_API FInterchangeTestFunction
+struct FInterchangeTestFunction
 {
 	GENERATED_BODY()
 
@@ -72,34 +74,34 @@ public:
 
 public:
 	/** Import all parameter data from text to a new binary buffer */
-	TSharedPtr<FStructOnScope> ImportParameters();
+	UE_API TSharedPtr<FStructOnScope> ImportParameters();
 
 	/** Import a single parameter from text to an existing binary buffer */
-	void ImportParameter(FStructOnScope& ParamData, FProperty* Property);
+	UE_API void ImportParameter(FStructOnScope& ParamData, FProperty* Property);
 
 	/** Export parameter data from binary to text */
-	void ExportParameters();
+	UE_API void ExportParameters();
 
 	/** Return a list of the supported asset classes */
-	static TArray<UClass*> GetAvailableAssetClasses();
+	static UE_API TArray<UClass*> GetAvailableAssetClasses();
 
 	/** Return a list of available functions which can act upon an asset of the given class */
-	static TArray<UFunction*> GetAvailableFunctions(const UClass* Class);
+	static UE_API TArray<UFunction*> GetAvailableFunctions(const UClass* Class);
 
 	/** Return whether the given function is valid or not */
-	bool IsValid(UFunction* Function) const;
+	UE_API bool IsValid(UFunction* Function) const;
 
 	/** Return whether the function takes an array of assets */
-	bool RequiresArrayOfAssets() const;
+	UE_API bool RequiresArrayOfAssets() const;
 
 	/** Return the class type expected as the first parameter to the function */
-	UClass* GetExpectedClass() const;
+	UE_API UClass* GetExpectedClass() const;
 
 	/** Returns the number of parameters which must be bound to this function call */
 	int32 GetNumParameters() const { return CheckFunction->NumParms - 2; }
 
 	/** Call the test function bound with the held parameter values */
-	FInterchangeTestFunctionResult Invoke(const TArray<UObject*>& AssetsToTest);
+	UE_API FInterchangeTestFunctionResult Invoke(const TArray<UObject*>& AssetsToTest);
 
 	/**
 	 * Provide a readable user-facing way of iterating the function parameters of interest.
@@ -168,3 +170,5 @@ public:
 	/** Return a proxy object which is used to iterate through function parameters of interest */
 	FParameters GetParameters() const { return FParameters(CheckFunction); } 
 };
+
+#undef UE_API
diff --git a/Engine/Plugins/Tests/InterchangeTests/Source/InterchangeTests/Public/InterchangeTestsModule.h b/Engine/Plugins/Tests/InterchangeTests/Source/InterchangeTests/Public/InterchangeTestsModule.h
--- a/Engine/Plugins/Tests/InterchangeTests/Source/InterchangeTests/Public/InterchangeTestsModule.h
+++ b/Engine/Plugins/Tests/InterchangeTests/Source/InterchangeTests/Public/InterchangeTestsModule.h
@@ -3,17 +3,21 @@
 
 #include "Modules/ModuleInterface.h"
 
+#define UE_API INTERCHANGETESTS_API
+
 #define INTERCHANGETESTS_MODULE_NAME TEXT("InterchangeTests")
 
 /**
  * Module for implementing Interchange automation tests
  */
-class INTERCHANGETESTS_API FInterchangeTestsModule : public IModuleInterface
+class FInterchangeTestsModule : public IModuleInterface
 {
 public:
-	static FInterchangeTestsModule& Get();
-	static bool IsAvailable();
+	static UE_API FInterchangeTestsModule& Get();
+	static UE_API bool IsAvailable();
 
 private:
-	virtual void StartupModule() override;
+	UE_API virtual void StartupModule() override;
 };
+
+#undef UE_API
diff --git a/Engine/Plugins/Tests/PythonAutomationTest/Source/PythonAutomationTest/Public/PythonAutomationTest.h b/Engine/Plugins/Tests/PythonAutomationTest/Source/PythonAutomationTest/Public/PythonAutomationTest.h
--- a/Engine/Plugins/Tests/PythonAutomationTest/Source/PythonAutomationTest/Public/PythonAutomationTest.h
+++ b/Engine/Plugins/Tests/PythonAutomationTest/Source/PythonAutomationTest/Public/PythonAutomationTest.h
@@ -6,28 +6,32 @@
 
 #include "PythonAutomationTest.generated.h"
 
-UCLASS(meta = (ScriptName = "PyAutomationTest"))
-class PYTHONAUTOMATIONTEST_API UPyAutomationTestLibrary : public UBlueprintFunctionLibrary
+#define UE_API PYTHONAUTOMATIONTEST_API
+
+UCLASS(MinimalAPI, meta = (ScriptName = "PyAutomationTest"))
+class UPyAutomationTestLibrary : public UBlueprintFunctionLibrary
 {
 	GENERATED_BODY()
 public:
 
 	UFUNCTION(BlueprintCallable, Category = "PyAutomationTest")
-	static void SetIsRunningPyLatentCommand(bool isRunning);
+	static UE_API void SetIsRunningPyLatentCommand(bool isRunning);
 
 	UFUNCTION(BlueprintCallable, Category = "PyAutomationTest")
-	static bool GetIsRunningPyLatentCommand();
+	static UE_API bool GetIsRunningPyLatentCommand();
 
 	UFUNCTION(BlueprintCallable, Category = "PyAutomationTest")
-	static void SetPyLatentCommandTimeout(float Seconds);
+	static UE_API void SetPyLatentCommandTimeout(float Seconds);
 
 	UFUNCTION(BlueprintCallable, Category = "PyAutomationTest")
-	static float GetPyLatentCommandTimeout();
+	static UE_API float GetPyLatentCommandTimeout();
 
 	UFUNCTION(BlueprintCallable, Category = "PyAutomationTest")
-	static void ResetPyLatentCommand();
+	static UE_API void ResetPyLatentCommand();
 
 private:
-	static bool IsRunningPyLatentCommand;
-	static float PyLatentCommandTimeout;
+	static UE_API bool IsRunningPyLatentCommand;
+	static UE_API float PyLatentCommandTimeout;
 };
+
+#undef UE_API
diff --git a/Engine/Plugins/TraceUtilities/Source/EditorTraceUtilities/Public/IEditorTraceUtilitiesModule.h b/Engine/Plugins/TraceUtilities/Source/EditorTraceUtilities/Public/IEditorTraceUtilitiesModule.h
--- a/Engine/Plugins/TraceUtilities/Source/EditorTraceUtilities/Public/IEditorTraceUtilitiesModule.h
+++ b/Engine/Plugins/TraceUtilities/Source/EditorTraceUtilities/Public/IEditorTraceUtilitiesModule.h
@@ -18,7 +18,7 @@ struct FStatusBarTraceSettings
 	ETraceDestination TraceDestination = ETraceDestination::TraceStore;
 };
 
-class EDITORTRACEUTILITIES_API IEditorTraceUtilitiesModule : public IModuleInterface
+class IEditorTraceUtilitiesModule : public IModuleInterface
 {
 public:
 
@@ -48,4 +48,4 @@ public:
 	/** Gets the tracing settings the status bar extension manages. */
 	virtual const FStatusBarTraceSettings& GetTraceSettings() const = 0;
 };
-} // namespace UE::EditorTraceUtilities
\ No newline at end of file
+} // namespace UE::EditorTraceUtilities
diff --git a/Engine/Plugins/TraceUtilities/Source/EditorTraceUtilities/Public/UnrealInsightsLauncher.h b/Engine/Plugins/TraceUtilities/Source/EditorTraceUtilities/Public/UnrealInsightsLauncher.h
--- a/Engine/Plugins/TraceUtilities/Source/EditorTraceUtilities/Public/UnrealInsightsLauncher.h
+++ b/Engine/Plugins/TraceUtilities/Source/EditorTraceUtilities/Public/UnrealInsightsLauncher.h
@@ -6,7 +6,9 @@
 #include "Async/TaskGraphInterfaces.h"
 #include "CoreMinimal.h"
 
-class EDITORTRACEUTILITIES_API FUnrealInsightsLauncher : public TSharedFromThis<FUnrealInsightsLauncher>
+#define UE_API EDITORTRACEUTILITIES_API
+
+class FUnrealInsightsLauncher : public TSharedFromThis<FUnrealInsightsLauncher>
 {
 	friend class FLogMessageOnGameThreadTask;
 
@@ -20,10 +22,10 @@ public:
 
 	typedef TFunction<void(const EStartInsightsResult /*Result*/)> StartUnrealInsightsCallback;
 
-	FUnrealInsightsLauncher();
-	~FUnrealInsightsLauncher();
+	UE_API FUnrealInsightsLauncher();
+	UE_API ~FUnrealInsightsLauncher();
 
-	static TSharedPtr<FUnrealInsightsLauncher> Instance;
+	static UE_API TSharedPtr<FUnrealInsightsLauncher> Instance;
 	
 	static const TSharedPtr<FUnrealInsightsLauncher>& Get()
 	{
@@ -37,7 +39,7 @@ public:
 	/**
 	 * Returns the full path to UnrealInsights.exe for the current engine installation 
 	 */
-	FString GetInsightsApplicationPath();
+	UE_API FString GetInsightsApplicationPath();
 
 	/**
 	 * Launches the UnrealInsights executable from the given Path, displays an editor message if it fails.
@@ -45,7 +47,7 @@ public:
 	 * @param Path The full filename of UnrealInsights.exe to launch
 	 * @param Parameters The command line parameters to use when launching the exe
 	 */	
-	void StartUnrealInsights(const FString& Path, const FString& Parameters = TEXT(""));
+	UE_API void StartUnrealInsights(const FString& Path, const FString& Parameters = TEXT(""));
 
 	/**
 	* Launches the UnrealInsights executable from the given Path, displays an editor message if it fails.
@@ -54,12 +56,12 @@ public:
 	* @param Parameters The command line parameters to use when launching the exe
 	* @param Callback A Callback that will be called when the launch process is completed.
 	*/
-	void StartUnrealInsights(const FString& Path, const FString& Parameters, StartUnrealInsightsCallback Callback);
+	UE_API void StartUnrealInsights(const FString& Path, const FString& Parameters, StartUnrealInsightsCallback Callback);
 
 	/**
 	* Closes UnrealInsights.exe.
 	*/
-	void CloseUnrealInsights();
+	UE_API void CloseUnrealInsights();
 
 	/**
 	 * Try to open a trace file from the trace target returned by UE::Trace::GetTraceDestination().
@@ -67,14 +69,14 @@ public:
 	 * file name of the active/live trace through the storebrowser API and open the remote trace in Insights.
 	 * @param Destination The trace destination of the currently running trace session.
 	 */
-	bool TryOpenTraceFromDestination(const FString& Destination );
+	UE_API bool TryOpenTraceFromDestination(const FString& Destination );
 
 	/**
 	 * Launches UnrealInsights.exe with FilePath as the only parameter
 	 * @param  FilePath Full file name and path to a .utrace file on the local machine
 	 * @return 
 	 */
-	bool OpenTraceFile(const FString& FilePath);
+	UE_API bool OpenTraceFile(const FString& FilePath);
 
 	/**
 	 * Launches UnrealInsights.exe with the -Store and -OpenTraceID parameters
@@ -83,26 +85,26 @@ public:
 	 * @param  TraceID ID of the Trace to open, can be obtained through the UE::Trace::FStoreBrowser interface
 	 * @return 
 	 */
-	bool OpenRemoteTrace(const FString& TraceHostAddress, const uint16 TraceHostPort, uint32 TraceID);
+	UE_API bool OpenRemoteTrace(const FString& TraceHostAddress, const uint16 TraceHostPort, uint32 TraceID);
 
 	/**
 	* This function tries to query the file name of the active/live trace through the storebrowser API
 	 * @param Hostname hostname or ip of the target UnrealTraceServerX
 	 */
-	bool OpenActiveTraceFromStore(const FString& TraceHostAddress);
+	UE_API bool OpenActiveTraceFromStore(const FString& TraceHostAddress);
 
 private:
 	/*
 	 * Attempts building UnrealInsights via UAT, will launch with forwarded parameters if successful.
 	 * Assumes that the Insights Executable in path belongs to this engine.
 	 */
-	void BuildUnrealInsights(const FString& Path, const FString& LaunchParameters, StartUnrealInsightsCallback Callback);
+	UE_API void BuildUnrealInsights(const FString& Path, const FString& LaunchParameters, StartUnrealInsightsCallback Callback);
 
 	/// Logs an error message to the MessageLog window
-	void LogMessage(const FText& Message);
+	UE_API void LogMessage(const FText& Message);
 
 	/// Logs an error message to the MessageLog window using a task that is ran on the Game Thread.
-	void LogMessageOnGameThread(const FText& Message);
+	UE_API void LogMessageOnGameThread(const FText& Message);
 
 private:
 	/** The name of the Unreal Insights log listing. */
@@ -144,3 +146,4 @@ private:
 	FGraphEventRef Event;
 };
 
+#undef UE_API
diff --git a/Engine/Plugins/UbaController/Source/Public/UbaControllerModule.h b/Engine/Plugins/UbaController/Source/Public/UbaControllerModule.h
--- a/Engine/Plugins/UbaController/Source/Public/UbaControllerModule.h
+++ b/Engine/Plugins/UbaController/Source/Public/UbaControllerModule.h
@@ -19,7 +19,7 @@ public:
 	UBACONTROLLER_API static FUbaControllerModule& Get();
 	UBACONTROLLER_API virtual bool IsSupported() override final;
 
-	UBACONTROLLER_API virtual const FString GetName() override final { return FString("UBA Controller"); };
+	virtual const FString GetName() override final { return FString("UBA Controller"); };
 	
 	virtual void StartupModule() override final;
 	virtual void ShutdownModule() override final;
