EngineTest:
- Add support for Substrate for MaterialX tests
- Add a message parameter for InterchangeErrorOrWarningTest
- Update MaterialX tests to avoid the warning message since now Substrate is enabled by default but not the Adaptive GBuffer

#virtualized

diff --git a/Engine/Plugins/Tests/InterchangeTests/Source/InterchangeTests/Private/ImportTestFunctions/InterchangeResultImportTestFunctions.cpp b/Engine/Plugins/Tests/InterchangeTests/Source/InterchangeTests/Private/ImportTestFunctions/InterchangeResultImportTestFunctions.cpp
--- a/Engine/Plugins/Tests/InterchangeTests/Source/InterchangeTests/Private/ImportTestFunctions/InterchangeResultImportTestFunctions.cpp
+++ b/Engine/Plugins/Tests/InterchangeTests/Source/InterchangeTests/Private/ImportTestFunctions/InterchangeResultImportTestFunctions.cpp
@@ -14,7 +14,7 @@ UClass* UInterchangeResultImportTestFunctions::GetAssociatedAssetType() const
 }
 
 
-FInterchangeTestFunctionResult UInterchangeResultImportTestFunctions::CheckIfErrorOrWarningWasGenerated(UInterchangeResultsContainer* ResultsContainer, TSubclassOf<UInterchangeResult> ErrorOrWarningClass)
+FInterchangeTestFunctionResult UInterchangeResultImportTestFunctions::CheckIfErrorOrWarningWasGenerated(UInterchangeResultsContainer* ResultsContainer, TSubclassOf<UInterchangeResult> ErrorOrWarningClass, const FString& Message)
 {
 	FInterchangeTestFunctionResult Result;
 
@@ -23,11 +23,15 @@ FInterchangeTestFunctionResult UInterchangeResultImportTestFunctions::CheckIfErr
 	{
 		if (ResultObject->IsA(ErrorOrWarningClass))
 		{
-			bFoundType = true;
-			// If we matched the message type, remove it from the results container so it won't flag as an error at the end of the test.
-			// Note: we can do this while iterating through results, because UInterchangeResultsContainer::GetResults passes the array by value,
-			// for thread-safety purposes.
-			ResultsContainer->RemoveResult(ResultObject);
+			if (ResultObject->GetText().ToString() == Message)
+			{
+				bFoundType = true;
+				// If we matched the message type, remove it from the results container so it won't flag as an error at the end of the test.
+				// Note: we can do this while iterating through results, because UInterchangeResultsContainer::GetResults passes the array by value,
+				// for thread-safety purposes.
+				ResultsContainer->RemoveResult(ResultObject);
+
+			}
 		}
 	}
 
diff --git a/Engine/Plugins/Tests/InterchangeTests/Source/InterchangeTests/Private/ImportTestFunctions/MaterialXTestFunctions.cpp b/Engine/Plugins/Tests/InterchangeTests/Source/InterchangeTests/Private/ImportTestFunctions/MaterialXTestFunctions.cpp
--- a/Engine/Plugins/Tests/InterchangeTests/Source/InterchangeTests/Private/ImportTestFunctions/MaterialXTestFunctions.cpp
+++ b/Engine/Plugins/Tests/InterchangeTests/Source/InterchangeTests/Private/ImportTestFunctions/MaterialXTestFunctions.cpp
@@ -13,14 +13,19 @@ namespace UE::Interchange::Tests::Private
 {
 	void GetMaterialFunctionConnectedInputs(const UMaterial& Material, TArray<FString>& InputNames)
 	{
-		if (!ensure(Material.HasBaseColorConnected()))
+		bool bHasBaseColor = Material.HasBaseColorConnected();
+		bool bHasSubstrateFrontMaterial = Material.HasSubstrateFrontMaterialConnected();
+
+		// For MaterialX those 2 are mutually exclusive, but for backwards compatibility let's keep it that way
+		if (!ensure(bHasBaseColor || bHasSubstrateFrontMaterial))
 		{
 			return;
 		}
 
 #if WITH_EDITORONLY_DATA
 		const UMaterialEditorOnlyData* MaterialEditorOnly = Material.GetEditorOnlyData();
-		if (UMaterialExpressionMaterialFunctionCall* FunctionCall = Cast<UMaterialExpressionMaterialFunctionCall>(MaterialEditorOnly->BaseColor.Expression))
+		UMaterialExpression* Expression = bHasSubstrateFrontMaterial ? MaterialEditorOnly->FrontMaterial.Expression : MaterialEditorOnly->BaseColor.Expression;
+		if (UMaterialExpressionMaterialFunctionCall* FunctionCall = Cast<UMaterialExpressionMaterialFunctionCall>(Expression))
 		{
 			InputNames.Empty();
 
diff --git a/Engine/Plugins/Tests/InterchangeTests/Source/InterchangeTests/Public/ImportTestFunctions/InterchangeResultImportTestFunctions.h b/Engine/Plugins/Tests/InterchangeTests/Source/InterchangeTests/Public/ImportTestFunctions/InterchangeResultImportTestFunctions.h
--- a/Engine/Plugins/Tests/InterchangeTests/Source/InterchangeTests/Public/ImportTestFunctions/InterchangeResultImportTestFunctions.h
+++ b/Engine/Plugins/Tests/InterchangeTests/Source/InterchangeTests/Public/ImportTestFunctions/InterchangeResultImportTestFunctions.h
@@ -26,7 +26,7 @@ public:
 
 	/** Check whether the specified InterchangeResult was emitted during import */
 	UFUNCTION(Exec)
-	static UE_API FInterchangeTestFunctionResult CheckIfErrorOrWarningWasGenerated(UInterchangeResultsContainer* ResultsContainer, TSubclassOf<UInterchangeResult> ErrorOrWarningClass);
+	static UE_API FInterchangeTestFunctionResult CheckIfErrorOrWarningWasGenerated(UInterchangeResultsContainer* ResultsContainer, TSubclassOf<UInterchangeResult> ErrorOrWarningClass, const FString& Message);
 };
 
 #undef UE_API
