Fixed compilation error with InterchangeTests plugin

#rnx

diff --git a/Engine/Plugins/Tests/InterchangeTests/InterchangeTests.uplugin b/Engine/Plugins/Tests/InterchangeTests/InterchangeTests.uplugin
--- a/Engine/Plugins/Tests/InterchangeTests/InterchangeTests.uplugin
+++ b/Engine/Plugins/Tests/InterchangeTests/InterchangeTests.uplugin
@@ -15,6 +15,11 @@
 	"IsBetaVersion": true,
 	"IsExperimentalVersion": false,
 	"Installed": false,
+	"SupportedTargetPlatforms": [
+		"Win64",
+		"Linux",
+		"Mac"
+	],
 	"Modules": [
 		{
 			"Name": "InterchangeTests",
