Interchange tests:
- Avoid failing test if the message is empty
- re-enable TP_FileNotFound

#virtualized

diff --git a/Engine/Plugins/Tests/InterchangeTests/Source/InterchangeTests/Private/ImportTestFunctions/InterchangeResultImportTestFunctions.cpp b/Engine/Plugins/Tests/InterchangeTests/Source/InterchangeTests/Private/ImportTestFunctions/InterchangeResultImportTestFunctions.cpp
--- a/Engine/Plugins/Tests/InterchangeTests/Source/InterchangeTests/Private/ImportTestFunctions/InterchangeResultImportTestFunctions.cpp
+++ b/Engine/Plugins/Tests/InterchangeTests/Source/InterchangeTests/Private/ImportTestFunctions/InterchangeResultImportTestFunctions.cpp
@@ -23,14 +23,16 @@ FInterchangeTestFunctionResult UInterchangeResultImportTestFunctions::CheckIfErr
 	{
 		if (ResultObject->IsA(ErrorOrWarningClass))
 		{
-			if (ResultObject->GetText().ToString() == Message)
+			bFoundType = true;
+			// If we matched the message type, remove it from the results container so it won't flag as an error at the end of the test.
+			// Note: we can do this while iterating through results, because UInterchangeResultsContainer::GetResults passes the array by value,
+			// for thread-safety purposes.
+			ResultsContainer->RemoveResult(ResultObject);
+
+			const FString ResultMessage = ResultObject->GetText().ToString();
+			if (!Message.IsEmpty() && ResultMessage != Message)
 			{
-				bFoundType = true;
-				// If we matched the message type, remove it from the results container so it won't flag as an error at the end of the test.
-				// Note: we can do this while iterating through results, because UInterchangeResultsContainer::GetResults passes the array by value,
-				// for thread-safety purposes.
-				ResultsContainer->RemoveResult(ResultObject);
-
+				Result.AddError(FString::Printf(TEXT("Expected '%s' error message text, found '%s'."), *Message, *ResultMessage));
 			}
 		}
 	}
