[Interchange] [Tests] Added ability to test for specific Texture types

#rnx
#Synced-CL 48609198

diff --git a/Engine/Plugins/Tests/InterchangeTests/Source/InterchangeTests/Private/ImportTestFunctions/TextureImportTestFunctions.cpp b/Engine/Plugins/Tests/InterchangeTests/Source/InterchangeTests/Private/ImportTestFunctions/TextureImportTestFunctions.cpp
--- a/Engine/Plugins/Tests/InterchangeTests/Source/InterchangeTests/Private/ImportTestFunctions/TextureImportTestFunctions.cpp
+++ b/Engine/Plugins/Tests/InterchangeTests/Source/InterchangeTests/Private/ImportTestFunctions/TextureImportTestFunctions.cpp
@@ -13,6 +13,26 @@ UClass* UTextureImportTestFunctions::GetAssociatedAssetType() const
 	return UTexture::StaticClass();
 }
 
+FInterchangeTestFunctionResult UTextureImportTestFunctions::CheckTextureType(const TArray<UTexture*>& Textures, TSubclassOf<UTexture> ExpectedTextureType, int32 ExpectedNumber)
+{
+	FInterchangeTestFunctionResult Result;
+	int32 Count = 0;
+	for (UTexture* Texture : Textures)
+	{
+		if (Texture->IsA(ExpectedTextureType.Get()))
+		{
+			Count++;
+		}
+	}
+
+	if (Count != ExpectedNumber)
+	{
+		Result.AddError(FString::Printf(TEXT("Expected to have %d textures of type %s, imported %d."), ExpectedNumber, *ExpectedTextureType.Get()->GetName(), Count));
+	}
+
+	return Result;
+}
+
 FInterchangeTestFunctionResult UTextureImportTestFunctions::CheckImportedTextureCount(const TArray<UTexture*>& Textures, int32 ExpectedNumberOfImportedTextures)
 {
 	FInterchangeTestFunctionResult Result;
diff --git a/Engine/Plugins/Tests/InterchangeTests/Source/InterchangeTests/Public/ImportTestFunctions/TextureImportTestFunctions.h b/Engine/Plugins/Tests/InterchangeTests/Source/InterchangeTests/Public/ImportTestFunctions/TextureImportTestFunctions.h
--- a/Engine/Plugins/Tests/InterchangeTests/Source/InterchangeTests/Public/ImportTestFunctions/TextureImportTestFunctions.h
+++ b/Engine/Plugins/Tests/InterchangeTests/Source/InterchangeTests/Public/ImportTestFunctions/TextureImportTestFunctions.h
@@ -23,6 +23,10 @@ public:
 	// UImportTestFunctionsBase interface
 	UE_API virtual UClass* GetAssociatedAssetType() const override;
 
+	/** Check whether the imported textures are of a specific type */
+	UFUNCTION(Exec)
+	static UE_API FInterchangeTestFunctionResult CheckTextureType(const TArray<UTexture*>& Textures, TSubclassOf<UTexture> ExpectedTextureType, int32 ExpectedNumber);
+
 	/** Check whether the expected number of textures are imported */
 	UFUNCTION(Exec)
 	static UE_API FInterchangeTestFunctionResult CheckImportedTextureCount(const TArray<UTexture*>& Textures, int32 ExpectedNumberOfImportedTextures);
