[Interchange] Toast is not shown at all when doing automation testing.
- Also added test plan assets for the Invisible actor tests.
- Added ability for the test plan asset to specify if the screenshot will force enable game view.

#rnx

#virtualized
#Synced-CL 51164812

diff --git a/Engine/Plugins/Interchange/Editor/Source/InterchangeEditor/Private/InterchangeEditorModule.cpp b/Engine/Plugins/Interchange/Editor/Source/InterchangeEditor/Private/InterchangeEditorModule.cpp
--- a/Engine/Plugins/Interchange/Editor/Source/InterchangeEditor/Private/InterchangeEditorModule.cpp
+++ b/Engine/Plugins/Interchange/Editor/Source/InterchangeEditor/Private/InterchangeEditorModule.cpp
@@ -127,6 +127,11 @@ namespace UE::Interchange::InterchangeEditorModule
 
 	void ImportedObjectAvailable(TMap<int32, TArray<TObjectPtr<UObject>>> ImportedAssetsPerSourceIndex, TMap<int32, TArray<TObjectPtr<UObject>>> ImportedSceneObjectsPerSourceIndex)
 	{
+		if (GIsAutomationTesting)
+		{
+			return;
+		}
+
 		bool bHasHiddenActors = false;
 		for (const TPair<int32, TArray<TObjectPtr<UObject>>>& SourceIndexSceneObjectsPair : ImportedSceneObjectsPerSourceIndex)
 		{
diff --git a/Engine/Plugins/Tests/InterchangeTests/Source/InterchangeTests/Private/InterchangeAutomatedTestUtils.cpp b/Engine/Plugins/Tests/InterchangeTests/Source/InterchangeTests/Private/InterchangeAutomatedTestUtils.cpp
--- a/Engine/Plugins/Tests/InterchangeTests/Source/InterchangeTests/Private/InterchangeAutomatedTestUtils.cpp
+++ b/Engine/Plugins/Tests/InterchangeTests/Source/InterchangeTests/Private/InterchangeAutomatedTestUtils.cpp
@@ -276,8 +276,20 @@ bool FInterchangeCaptureScreenshotCommand::Update()
 	if (TestStepData->ScreenshotTask == nullptr)
 	{
 		using namespace InterchangeAutomationTestUtils::Private;
-
-		TestStepData->ScreenshotTask = UAutomationBlueprintFunctionLibrary::TakeHighResScreenshot(1280, 720, GetScreenshotNameString(*TestStepData), nullptr, false, false, ScreenshotParameters.ComparisonTolerance);
+		
+		constexpr bool bMaskEnabled = false;
+		constexpr bool bCaptureHDR = false;
+		constexpr float Delay = 0.0f;
+
+		TestStepData->ScreenshotTask = UAutomationBlueprintFunctionLibrary::TakeHighResScreenshot(
+			1280, 720, 
+			GetScreenshotNameString(*TestStepData), 
+			nullptr, 
+			bMaskEnabled, bCaptureHDR, 
+			ScreenshotParameters.ComparisonTolerance,
+			TEXT(""),
+			Delay,
+			ScreenshotParameters.bForceGameView);
 	}
 
 	if (TestStepData->ScreenshotTask && !TestStepData->ScreenshotTask->IsValidTask())
diff --git a/Engine/Plugins/Tests/InterchangeTests/Source/InterchangeTests/Public/InterchangeImportTestStepBase.h b/Engine/Plugins/Tests/InterchangeTests/Source/InterchangeTests/Public/InterchangeImportTestStepBase.h
--- a/Engine/Plugins/Tests/InterchangeTests/Source/InterchangeTests/Public/InterchangeImportTestStepBase.h
+++ b/Engine/Plugins/Tests/InterchangeTests/Source/InterchangeTests/Public/InterchangeImportTestStepBase.h
@@ -52,6 +52,9 @@ public:
 
 	UPROPERTY(EditAnywhere, Category = Viewport)
 	float WireframeOpacity = 0.2f;
+
+	UPROPERTY(EditAnywhere, Category = Viewport)
+	bool bForceGameView = true;
 };
 
 
